Fix missing include

diff --git a/Engine/Plugins/Editor/GuidedTutorials/Source/IntroTutorials/Private/IntroTutorials.cpp b/Engine/Plugins/Editor/GuidedTutorials/Source/IntroTutorials/Private/IntroTutorials.cpp
--- a/Engine/Plugins/Editor/GuidedTutorials/Source/IntroTutorials/Private/IntroTutorials.cpp
+++ b/Engine/Plugins/Editor/GuidedTutorials/Source/IntroTutorials/Private/IntroTutorials.cpp
@@ -18,7 +18,6 @@
 
 #include "EditorTutorial.h"
 #include "EditorTutorialSettings.h"
-#include "EditorTutorialStyleSet.h"
 #include "TutorialStateSettings.h"
 #include "TutorialSettings.h"
 #include "ISettingsModule.h"
