Editor tutorials to plugin, part 2: move styles and content.

diff --git a/Engine/Plugins/Editor/GuidedTutorials/Source/IntroTutorials/IntroTutorials.Build.cs b/Engine/Plugins/Editor/GuidedTutorials/Source/IntroTutorials/IntroTutorials.Build.cs
--- a/Engine/Plugins/Editor/GuidedTutorials/Source/IntroTutorials/IntroTutorials.Build.cs
+++ b/Engine/Plugins/Editor/GuidedTutorials/Source/IntroTutorials/IntroTutorials.Build.cs
@@ -13,6 +13,7 @@ namespace UnrealBuildTool.Rules
 					"Engine", // @todo Mac: for some reason CoreUObject and Engine are needed to link in debug on Mac
                     "InputCore",
                     "Slate",
+					"SlateCore",
 					"EditorStyle",
                     "Documentation",
 					"GraphEditor",
diff --git a/Engine/Plugins/Editor/GuidedTutorials/Source/IntroTutorials/Private/ClassTypeActions_EditorTutorial.cpp b/Engine/Plugins/Editor/GuidedTutorials/Source/IntroTutorials/Private/ClassTypeActions_EditorTutorial.cpp
--- a/Engine/Plugins/Editor/GuidedTutorials/Source/IntroTutorials/Private/ClassTypeActions_EditorTutorial.cpp
+++ b/Engine/Plugins/Editor/GuidedTutorials/Source/IntroTutorials/Private/ClassTypeActions_EditorTutorial.cpp
@@ -9,6 +9,7 @@
 #include "EditorStyleSet.h"
 #include "IIntroTutorials.h"
 #include "EditorTutorial.h"
+#include "EditorTutorialStyle.h"
 #include "AssetData.h"
 
 #define LOCTEXT_NAMESPACE "IntroTutorials"
@@ -37,7 +38,7 @@ TSharedPtr<SWidget> FClassTypeActions_EditorTutorial::GetThumbnailOverlay(const
 		.Padding(FMargin(2))
 		[
 			SNew(SButton)
-			.ButtonStyle(FEditorStyle::Get(), "HoverHintOnly")
+			.ButtonStyle(FAppStyle::Get(), "HoverHintOnly")
 			.ToolTipText(LOCTEXT("Blueprint_LaunchTutorialToolTip", "Launch this tutorial."))
 			.Cursor(EMouseCursor::Default) // The outer widget can specify a DragHand cursor, so we need to override that here
 			.ForegroundColor(FSlateColor::UseForeground())
@@ -48,7 +49,7 @@ TSharedPtr<SWidget> FClassTypeActions_EditorTutorial::GetThumbnailOverlay(const
 				.MinDesiredHeight(16)
 				[
 					SNew(SImage)
-					.Image(FEditorStyle::GetBrush("Tutorials.Browser.PlayButton.Image"))
+					.Image(FEditorTutorialStyle::Get().GetBrush("Tutorials.Browser.PlayButton.Image"))
 				]
 			]
 		];
diff --git a/Engine/Plugins/Editor/GuidedTutorials/Source/IntroTutorials/Private/EditorTutorialStyle.cpp b/Engine/Plugins/Editor/GuidedTutorials/Source/IntroTutorials/Private/EditorTutorialStyle.cpp
new file mode 100644
--- /dev/null
+++ b/Engine/Plugins/Editor/GuidedTutorials/Source/IntroTutorials/Private/EditorTutorialStyle.cpp
@@ -0,0 +1,349 @@
+// Copyright Epic Games, Inc. All Rights Reserved.
+
+#include "EditorTutorialStyle.h"
+#include "Styling/SlateStyleRegistry.h"
+#include "Styling/SlateStyleMacros.h"
+#include "Styling/SlateTypes.h"
+#include "Styling/AppStyle.h"
+
+FName FEditorTutorialStyle::StyleName("FEditorTutorialStyle");
+TUniquePtr<FEditorTutorialStyle> FEditorTutorialStyle::Instance(nullptr);
+
+const FName& FEditorTutorialStyle::GetStyleSetName() const
+{
+	return StyleName;
+}
+
+const FEditorTutorialStyle& FEditorTutorialStyle::Get()
+{
+	if (!Instance.IsValid())
+	{
+		Instance = TUniquePtr<FEditorTutorialStyle>(new FEditorTutorialStyle);
+	}
+	return *(Instance.Get());
+}
+
+void FEditorTutorialStyle::Shutdown()
+{
+	Instance.Reset();
+}
+
+FEditorTutorialStyle::FEditorTutorialStyle()
+	: FSlateStyleSet(StyleName)
+{
+	SetContentRoot(FPaths::EnginePluginsDir() / TEXT("Editor/GuidedTutorials/Content/Editor/Slate"));
+
+	const FVector2D Icon8x8(8.0f, 8.0f);
+	const FVector2D Icon16x16(16.0f, 16.0f);
+	const FVector2D Icon32x32(32.0f, 32.0f);
+	const FVector2D Icon128x128(128.0f, 128.0f);
+
+	Set("Tutorials.MenuIcon", new CORE_IMAGE_BRUSH_SVG("Starship/Common/Tutorials", Icon16x16));
+
+	const FLinearColor TutorialButtonColor = FLinearColor(0.15f, 0.15f, 0.15f, 1.0f);
+	const FLinearColor TutorialSelectionColor = FLinearColor(0.19f, 0.33f, 0.72f);
+	const FLinearColor TutorialNavigationButtonColor = FLinearColor(0.0f, 0.59f, 0.14f, 1.0f);
+	const FLinearColor TutorialNavigationButtonHoverColor = FLinearColor(0.2f, 0.79f, 0.34f, 1.0f);
+	const FLinearColor TutorialNavigationBackButtonColor = TutorialNavigationButtonColor;
+	const FLinearColor TutorialNavigationBackButtonHoverColor = TutorialNavigationButtonHoverColor;
+
+	const FTextBlockStyle TutorialText = FTextBlockStyle(FAppStyle::Get().GetWidgetStyle<FTextBlockStyle>("Documentation.Text"))
+		.SetColorAndOpacity(FLinearColor::Black)
+		.SetHighlightColor(TutorialSelectionColor);
+
+	const FTextBlockStyle TutorialHeaderText = FTextBlockStyle(FAppStyle::Get().GetWidgetStyle<FTextBlockStyle>("Documentation.Header.Text"))
+		.SetColorAndOpacity(FLinearColor::Black)
+		.SetHighlightColor(TutorialSelectionColor);
+
+	Set("Tutorials.Border", new BOX_BRUSH("Tutorials/OverlayFrame", FMargin(18.0f / 64.0f), FLinearColor(1.0f, 1.0f, 1.0f, 1.0f)));
+	Set("Tutorials.Shadow", new BOX_BRUSH("Tutorials/TutorialShadow", FVector2D(256.0f, 256.0f), FMargin(114.0f / 256.0f)));
+
+	const FTextBlockStyle TutorialBrowserText = FTextBlockStyle(TutorialText)
+		.SetColorAndOpacity(FSlateColor::UseForeground())
+		.SetHighlightColor(TutorialSelectionColor);
+
+	Set("Tutorials.Browser.Text", TutorialBrowserText);
+
+	Set("Tutorials.Browser.WelcomeHeader", FTextBlockStyle(TutorialBrowserText)
+		.SetFontSize(20));
+
+	Set("Tutorials.Browser.SummaryHeader", FTextBlockStyle(TutorialBrowserText)
+		.SetFontSize(16));
+
+	Set("Tutorials.Browser.SummaryText", FTextBlockStyle(TutorialBrowserText)
+		.SetFontSize(10));
+
+	Set("Tutorials.Browser.HighlightTextColor", TutorialSelectionColor);
+
+	Set("Tutorials.Browser.Button", FButtonStyle()
+		.SetNormal(CORE_BOX_BRUSH("Common/ButtonHoverHint", FMargin(4 / 16.0f), FLinearColor(0.05f, 0.05f, 0.05f, 1)))
+		.SetHovered(CORE_BOX_BRUSH("Common/ButtonHoverHint", FMargin(4 / 16.0f), FLinearColor(0.07f, 0.07f, 0.07f, 1)))
+		.SetPressed(CORE_BOX_BRUSH("Common/ButtonHoverHint", FMargin(4 / 16.0f), FLinearColor(0.08f, 0.08f, 0.08f, 1)))
+		.SetNormalPadding(FMargin(0, 0, 0, 1))
+		.SetPressedPadding(FMargin(0, 1, 0, 0)));
+
+	Set("Tutorials.Browser.BackButton", FButtonStyle()
+		.SetNormal(CORE_BOX_BRUSH("Common/ButtonHoverHint", FMargin(4 / 16.0f), FLinearColor(1.0f, 1.0f, 1.0f, 0.0f)))
+		.SetHovered(CORE_BOX_BRUSH("Common/ButtonHoverHint", FMargin(4 / 16.0f), FLinearColor(1.0f, 1.0f, 1.0f, 0.05f)))
+		.SetPressed(CORE_BOX_BRUSH("Common/ButtonHoverHint", FMargin(4 / 16.0f), FLinearColor(1.0f, 1.0f, 1.0f, 0.05f)))
+		.SetNormalPadding(FMargin(0, 0, 0, 1))
+		.SetPressedPadding(FMargin(0, 1, 0, 0)));
+
+	Set("Tutorials.Content.Button", FButtonStyle()
+		.SetNormal(CORE_BOX_BRUSH("Common/ButtonHoverHint", FMargin(4 / 16.0f), FLinearColor(0, 0, 0, 0)))
+		.SetHovered(CORE_BOX_BRUSH("Common/ButtonHoverHint", FMargin(4 / 16.0f), FLinearColor(1, 1, 1, 1)))
+		.SetPressed(CORE_BOX_BRUSH("Common/ButtonHoverHint", FMargin(4 / 16.0f), FLinearColor(1, 1, 1, 1)))
+		.SetNormalPadding(FMargin(0, 0, 0, 1))
+		.SetPressedPadding(FMargin(0, 1, 0, 0)));
+
+	Set("Tutorials.Content.NavigationButtonWrapper", FButtonStyle()
+		.SetNormal(FSlateNoResource())
+		.SetHovered(FSlateNoResource())
+		.SetPressed(FSlateNoResource())
+		.SetNormalPadding(FMargin(0, 0, 0, 1))
+		.SetPressedPadding(FMargin(0, 1, 0, 0)));
+
+	Set("Tutorials.Content.NavigationButton", FButtonStyle()
+		.SetNormal(CORE_BOX_BRUSH("Common/ButtonHoverHint", FMargin(4 / 16.0f), TutorialNavigationButtonColor))
+		.SetHovered(CORE_BOX_BRUSH("Common/ButtonHoverHint", FMargin(4 / 16.0f), TutorialNavigationButtonHoverColor))
+		.SetPressed(CORE_BOX_BRUSH("Common/ButtonHoverHint", FMargin(4 / 16.0f), TutorialNavigationButtonHoverColor))
+		.SetNormalPadding(FMargin(0, 0, 0, 1))
+		.SetPressedPadding(FMargin(0, 1, 0, 0)));
+
+	Set("Tutorials.Content.NavigationBackButton", FButtonStyle()
+		.SetNormal(CORE_BOX_BRUSH("Common/ButtonHoverHint", FMargin(4 / 16.0f), TutorialNavigationBackButtonColor))
+		.SetHovered(CORE_BOX_BRUSH("Common/ButtonHoverHint", FMargin(4 / 16.0f), TutorialNavigationBackButtonHoverColor))
+		.SetPressed(CORE_BOX_BRUSH("Common/ButtonHoverHint", FMargin(4 / 16.0f), TutorialNavigationBackButtonHoverColor))
+		.SetNormalPadding(FMargin(0, 0, 0, 1))
+		.SetPressedPadding(FMargin(0, 1, 0, 0)));
+
+	Set("Tutorials.Content.NavigationText", FTextBlockStyle(TutorialText));
+
+	Set("Tutorials.Content.Color", FLinearColor(1.0f, 1.0f, 1.0f, 0.9f));
+	Set("Tutorials.Content.Color.Hovered", FLinearColor(1.0f, 1.0f, 1.0f, 1.0f));
+
+	Set("Tutorials.Browser.CategoryArrow", new IMAGE_BRUSH("Tutorials/BrowserCategoryArrow", FVector2D(24.0f, 24.0f), FSlateColor::UseForeground()));
+	Set("Tutorials.Browser.DefaultTutorialIcon", new IMAGE_BRUSH("Tutorials/DefaultTutorialIcon_40x", FVector2D(40.0f, 40.0f), FLinearColor::White));
+	Set("Tutorials.Browser.DefaultCategoryIcon", new IMAGE_BRUSH("Tutorials/DefaultCategoryIcon_40x", FVector2D(40.0f, 40.0f), FLinearColor::White));
+
+	Set("Tutorials.Browser.BackButton.Image", new IMAGE_BRUSH("Tutorials/BrowserBack", FVector2D(32.0f, 32.0f), FLinearColor(1.0f, 1.0f, 1.0f, 1.0f)));
+	Set("Tutorials.Browser.PlayButton.Image", new IMAGE_BRUSH("Tutorials/BrowserPlay", FVector2D(32.0f, 32.0f), FLinearColor(1.0f, 1.0f, 1.0f, 1.0f)));
+	Set("Tutorials.Browser.RestartButton", new IMAGE_BRUSH("Tutorials/BrowserRestart", FVector2D(16.0f, 16.0f), FLinearColor(1.0f, 1.0f, 1.0f, 1.0f)));
+	
+	Set("Tutorials.Browser.Completed", new IMAGE_BRUSH("Tutorials/TutorialCompleted", Icon32x32));
+	Set("Tutorials.Browser.Breadcrumb", new IMAGE_BRUSH("Tutorials/Breadcrumb", Icon8x8, FLinearColor::White));
+	Set("Tutorials.Browser.PathText", FTextBlockStyle(TutorialBrowserText)
+		.SetFontSize(9));
+
+	Set("Tutorials.Navigation.Button", FButtonStyle()
+		.SetNormal(CORE_BOX_BRUSH("Common/ButtonHoverHint", FMargin(4 / 16.0f), FLinearColor(0, 0, 0, 0)))
+		.SetHovered(CORE_BOX_BRUSH("Common/ButtonHoverHint", FMargin(4 / 16.0f), FLinearColor(0, 0, 0, 0)))
+		.SetPressed(CORE_BOX_BRUSH("Common/ButtonHoverHint", FMargin(4 / 16.0f), FLinearColor(0, 0, 0, 0)))
+		.SetNormalPadding(FMargin(0, 0, 0, 1))
+		.SetPressedPadding(FMargin(0, 1, 0, 0)));
+	
+	Set("Tutorials.Navigation.NextButton", new IMAGE_BRUSH("Tutorials/NavigationNext", Icon32x32));
+	Set("Tutorials.Navigation.HomeButton", new IMAGE_BRUSH("Tutorials/NavigationHome", Icon32x32));
+	Set("Tutorials.Navigation.BackButton", new IMAGE_BRUSH("Tutorials/NavigationBack", Icon32x32));
+
+	Set("Tutorials.WidgetContent", FTextBlockStyle(TutorialText)
+		.SetFontSize(10));
+
+	Set("Tutorials.ButtonColor", TutorialButtonColor);
+	Set("Tutorials.ButtonHighlightColor", TutorialSelectionColor);
+	Set("Tutorials.ButtonDisabledColor", FAppStyle::Get().GetSlateColor("SelectionColor_Inactive"));
+	
+	Set("Tutorials.ContentAreaBackground", new BOX_BRUSH("Tutorials/TutorialContentBackground", FMargin(4 / 16.0f)));
+	Set("Tutorials.HomeContentAreaBackground", new BOX_BRUSH("Tutorials/TutorialHomeContentBackground", FMargin(4 / 16.0f)));
+	Set("Tutorials.ContentAreaFrame", new BOX_BRUSH("Tutorials/ContentAreaFrame", FMargin(26.0f / 64.0f)));
+	Set("Tutorials.CurrentExcerpt", new IMAGE_BRUSH("Tutorials/CurrentExcerpt", FVector2D(24.0f, 24.0f), TutorialSelectionColor));
+	Set("Tutorials.Home", new IMAGE_BRUSH("Tutorials/HomeButton", FVector2D(32.0f, 32.0f)));
+	Set("Tutorials.Back", new IMAGE_BRUSH("Tutorials/BackButton", FVector2D(24.0f, 24.0f)));
+	Set("Tutorials.Next", new IMAGE_BRUSH("Tutorials/NextButton", FVector2D(24.0f, 24.0f)));
+	
+	Set("Tutorials.PageHeader", FTextBlockStyle(TutorialHeaderText)
+		.SetFontSize(22));
+
+	Set("Tutorials.CurrentExcerpt", FTextBlockStyle(TutorialHeaderText)
+		.SetFontSize(16));
+
+	Set("Tutorials.NavigationButtons", FTextBlockStyle(TutorialHeaderText)
+		.SetFontSize(16));
+
+	// UDN documentation styles
+	Set("Tutorials.Content", FTextBlockStyle(TutorialText)
+		.SetColorAndOpacity(FSlateColor::UseForeground()));
+	Set("Tutorials.Hyperlink.Text", FTextBlockStyle(FAppStyle::Get().GetWidgetStyle<FTextBlockStyle>("Documentation.Hyperlink.Text")));
+	Set("Tutorials.NumberedContent", FTextBlockStyle(TutorialText));
+	Set("Tutorials.BoldContent", FTextBlockStyle(TutorialText)
+		.SetTypefaceFontName(TEXT("Bold")));
+
+	Set("Tutorials.Header1", FTextBlockStyle(TutorialHeaderText)
+		.SetFontSize(32));
+
+	Set("Tutorials.Header2", FTextBlockStyle(TutorialHeaderText)
+		.SetFontSize(24));
+	
+	Set("Tutorials.Hyperlink.Button", FButtonStyle(FAppStyle::Get().GetWidgetStyle<FButtonStyle>("Documentation.Hyperlink.Button"))
+		.SetNormal(CORE_BORDER_BRUSH("Old/HyperlinkDotted", FMargin(0, 0, 0, 3 / 16.0f), FLinearColor::Black))
+		.SetHovered(CORE_BORDER_BRUSH("Old/HyperlinkUnderline", FMargin(0, 0, 0, 3 / 16.0f), FLinearColor::Black)));
+
+	Set("Tutorials.Separator", new CORE_BOX_BRUSH("Common/Separator", 1 / 4.0f, FLinearColor::Black));
+
+	Set("Tutorials.ProgressBar", FProgressBarStyle()
+		.SetBackgroundImage(CORE_BOX_BRUSH("Common/ProgressBar_Background", FMargin(5.f / 12.f)))
+		.SetFillImage(CORE_BOX_BRUSH("Common/ProgressBar_NeutralFill", FMargin(5.f / 12.f)))
+		.SetMarqueeImage(CORE_IMAGE_BRUSH("Common/ProgressBar_Marquee", FVector2D(20, 12), FLinearColor::White, ESlateBrushTileType::Horizontal))
+	);
+
+	// Default text styles
+	{
+		const FTextBlockStyle RichTextNormal = FTextBlockStyle()
+			.SetFont(DEFAULT_FONT("Regular", 11))
+			.SetColorAndOpacity(FSlateColor::UseForeground())
+			.SetShadowOffset(FVector2D::ZeroVector)
+			.SetShadowColorAndOpacity(FLinearColor::Black)
+			.SetHighlightColor(FLinearColor(0.02f, 0.3f, 0.0f))
+			.SetHighlightShape(BOX_BRUSH("Common/TextBlockHighlightShape", FMargin(3.f / 8.f)));
+		Set("Tutorials.Content.Text", RichTextNormal);
+
+		Set("Tutorials.Content.TextBold", FTextBlockStyle(RichTextNormal)
+			.SetFont(DEFAULT_FONT("Bold", 11)));
+
+		Set("Tutorials.Content.HeaderText1", FTextBlockStyle(RichTextNormal)
+			.SetFontSize(20));
+
+		Set("Tutorials.Content.HeaderText2", FTextBlockStyle(RichTextNormal)
+			.SetFontSize(16));
+
+		{
+			const FButtonStyle RichTextHyperlinkButton = FButtonStyle()
+				.SetNormal(CORE_BORDER_BRUSH("Old/HyperlinkDotted", FMargin(0, 0, 0, 3 / 16.0f), FLinearColor::Blue))
+				.SetPressed(FSlateNoResource())
+				.SetHovered(CORE_BORDER_BRUSH("Old/HyperlinkUnderline", FMargin(0, 0, 0, 3 / 16.0f), FLinearColor::Blue));
+
+			const FTextBlockStyle RichTextHyperlinkText = FTextBlockStyle(RichTextNormal)
+				.SetColorAndOpacity(FLinearColor::Blue);
+
+			Set("Tutorials.Content.HyperlinkText", RichTextHyperlinkText);
+
+			// legacy style
+			Set("TutorialEditableText.Editor.HyperlinkText", RichTextHyperlinkText);
+
+			const FHyperlinkStyle RichTextHyperlink = FHyperlinkStyle()
+				.SetUnderlineStyle(RichTextHyperlinkButton)
+				.SetTextStyle(RichTextHyperlinkText)
+				.SetPadding(FMargin(0.0f));
+			Set("Tutorials.Content.Hyperlink", RichTextHyperlink);
+
+			Set("Tutorials.Content.ExternalLink", new IMAGE_BRUSH("Tutorials/ExternalLink", Icon16x16, FLinearColor::Blue));
+
+			// legacy style
+			Set("TutorialEditableText.Editor.Hyperlink", RichTextHyperlink);
+		}
+	}
+
+	// Toolbar
+	{
+		const FLinearColor NormalColor(FColor(0xffeff3f3));
+		const FLinearColor SelectedColor(FColor(0xffdbe4d5));
+		const FLinearColor HoverColor(FColor(0xffdbe4e4));
+		const FLinearColor DisabledColor(FColor(0xaaaaaa));
+		const FLinearColor TextColor(FColor(0xff2c3e50));
+
+		Set("TutorialEditableText.RoundedBackground", new BOX_BRUSH("Common/RoundedSelection_16x", 4.0f / 16.0f, FLinearColor(FColor(0xffeff3f3))));
+
+		Set("TutorialEditableText.Toolbar.HyperlinkImage", new IMAGE_BRUSH("Tutorials/hyperlink", Icon16x16, TextColor));
+		Set("TutorialEditableText.Toolbar.ImageImage", new IMAGE_BRUSH("Tutorials/Image", Icon16x16, TextColor));
+
+		Set("TutorialEditableText.Toolbar.TextColor", TextColor);
+
+		Set("TutorialEditableText.Toolbar.Text", FTextBlockStyle(FAppStyle::Get().GetWidgetStyle<FTextBlockStyle>("NormalText"))
+			.SetFont(DEFAULT_FONT("Regular", 10))
+			.SetColorAndOpacity(TextColor)
+		);
+
+		Set("TutorialEditableText.Toolbar.BoldText", FTextBlockStyle(FAppStyle::Get().GetWidgetStyle<FTextBlockStyle>("NormalText"))
+			.SetFont(DEFAULT_FONT("Bold", 10))
+			.SetColorAndOpacity(TextColor)
+		);
+
+		Set("TutorialEditableText.Toolbar.ItalicText", FTextBlockStyle(FAppStyle::Get().GetWidgetStyle<FTextBlockStyle>("NormalText"))
+			.SetFont(DEFAULT_FONT("Italic", 10))
+			.SetColorAndOpacity(TextColor)
+		);
+
+		Set("TutorialEditableText.Toolbar.Checkbox", FCheckBoxStyle()
+			.SetCheckBoxType(ESlateCheckBoxType::CheckBox)
+			.SetUncheckedImage(IMAGE_BRUSH("Common/CheckBox", Icon16x16, FLinearColor::White))
+			.SetUncheckedHoveredImage(IMAGE_BRUSH("Common/CheckBox", Icon16x16, HoverColor))
+			.SetUncheckedPressedImage(IMAGE_BRUSH("Common/CheckBox_Hovered", Icon16x16, HoverColor))
+			.SetCheckedImage(IMAGE_BRUSH("Common/CheckBox_Checked_Hovered", Icon16x16, FLinearColor::White))
+			.SetCheckedHoveredImage(IMAGE_BRUSH("Common/CheckBox_Checked_Hovered", Icon16x16, HoverColor))
+			.SetCheckedPressedImage(IMAGE_BRUSH("Common/CheckBox_Checked", Icon16x16, HoverColor))
+			.SetUndeterminedImage(IMAGE_BRUSH("Common/CheckBox_Undetermined", Icon16x16, FLinearColor::White))
+			.SetUndeterminedHoveredImage(IMAGE_BRUSH("Common/CheckBox_Undetermined_Hovered", Icon16x16, HoverColor))
+			.SetUndeterminedPressedImage(IMAGE_BRUSH("Common/CheckBox_Undetermined_Hovered", Icon16x16, FLinearColor::White))
+		);
+
+		Set("TutorialEditableText.Toolbar.ToggleButtonCheckbox", FCheckBoxStyle()
+			.SetCheckBoxType(ESlateCheckBoxType::ToggleButton)
+			.SetUncheckedImage(BOX_BRUSH("Tutorials/FlatColorSquare", FVector2D(1.0f, 1.0f), FMargin(1), NormalColor))
+			.SetUncheckedHoveredImage(BOX_BRUSH("Tutorials/FlatColorSquare", FVector2D(1.0f, 1.0f), FMargin(1), HoverColor))
+			.SetUncheckedPressedImage(BOX_BRUSH("Tutorials/FlatColorSquare", FVector2D(1.0f, 1.0f), FMargin(1), HoverColor))
+			.SetCheckedImage(BOX_BRUSH("Tutorials/FlatColorSquare", FVector2D(1.0f, 1.0f), FMargin(1), SelectedColor))
+			.SetCheckedHoveredImage(BOX_BRUSH("Tutorials/FlatColorSquare", FVector2D(1.0f, 1.0f), FMargin(1), HoverColor))
+			.SetCheckedPressedImage(BOX_BRUSH("Tutorials/FlatColorSquare", FVector2D(1.0f, 1.0f), FMargin(1), HoverColor))
+		);
+
+		const FButtonStyle TutorialButton = FButtonStyle()
+			.SetNormal(BOX_BRUSH("Tutorials/FlatColorSquare", FVector2D(1.0f, 1.0f), FMargin(1), NormalColor))
+			.SetHovered(BOX_BRUSH("Tutorials/FlatColorSquare", FVector2D(1.0f, 1.0f), FMargin(1), HoverColor))
+			.SetPressed(BOX_BRUSH("Tutorials/FlatColorSquare", FVector2D(1.0f, 1.0f), FMargin(1), SelectedColor))
+			.SetNormalPadding(FMargin(2, 2, 2, 2))
+			.SetPressedPadding(FMargin(2, 3, 2, 1));
+		Set("TutorialEditableText.Toolbar.Button", TutorialButton);
+
+		const FComboButtonStyle ComboButton = FComboButtonStyle()
+			.SetButtonStyle(FButtonStyle())
+			.SetDownArrowImage(IMAGE_BRUSH("Common/ComboArrow", Icon8x8))
+			.SetMenuBorderBrush(BOX_BRUSH("Tutorials/FlatColorSquare", FVector2D(1.0f, 1.0f), FMargin(1), NormalColor))
+			.SetMenuBorderPadding(FMargin(0.0f));
+		Set("TutorialEditableText.Toolbar.ComboButton", ComboButton);
+
+		{
+			const FButtonStyle ComboBoxButton = FButtonStyle()
+				.SetNormal(BOX_BRUSH("Tutorials/FlatColorSquare", FVector2D(1.0f, 1.0f), FMargin(1), FLinearColor::White))
+				.SetHovered(BOX_BRUSH("Tutorials/FlatColorSquare", FVector2D(1.0f, 1.0f), FMargin(1), FLinearColor::White))
+				.SetPressed(BOX_BRUSH("Tutorials/FlatColorSquare", FVector2D(1.0f, 1.0f), FMargin(1), FLinearColor::White))
+				.SetNormalPadding(FMargin(2, 2, 2, 2))
+				.SetPressedPadding(FMargin(2, 3, 2, 1));
+
+			const FComboButtonStyle ComboBoxComboButton = FComboButtonStyle(ComboButton)
+				.SetButtonStyle(ComboBoxButton)
+				.SetMenuBorderPadding(FMargin(1.0));
+
+			Set("TutorialEditableText.Toolbar.ComboBox", FComboBoxStyle()
+				.SetComboButtonStyle(ComboBoxComboButton)
+			);
+		}
+	}
+
+	Set("TutorialLaunch.Button", FButtonStyle()
+		.SetNormalPadding(0)
+		.SetPressedPadding(0)
+		.SetNormal(IMAGE_BRUSH("Tutorials/TutorialButton_Default_16x", Icon16x16))
+		.SetHovered(IMAGE_BRUSH("Tutorials/TutorialButton_Hovered_16x", Icon16x16))
+		.SetPressed(IMAGE_BRUSH("Tutorials/TutorialButton_Pressed_16x", Icon16x16))
+	);
+
+	Set("TutorialLaunch.Circle", new IMAGE_BRUSH("Tutorials/Circle_128x", Icon128x128, FLinearColor::White));
+	Set("TutorialLaunch.Circle.Color", FLinearColor::Green);
+
+	FSlateStyleRegistry::RegisterSlateStyle(*this);
+}
+
+FEditorTutorialStyle::~FEditorTutorialStyle()
+{
+	FSlateStyleRegistry::UnRegisterSlateStyle(*this);
+}
diff --git a/Engine/Plugins/Editor/GuidedTutorials/Source/IntroTutorials/Private/EditorTutorialStyle.h b/Engine/Plugins/Editor/GuidedTutorials/Source/IntroTutorials/Private/EditorTutorialStyle.h
new file mode 100644
--- /dev/null
+++ b/Engine/Plugins/Editor/GuidedTutorials/Source/IntroTutorials/Private/EditorTutorialStyle.h
@@ -0,0 +1,25 @@
+// Copyright Epic Games, Inc. All Rights Reserved.
+
+#pragma once
+
+#include "CoreMinimal.h"
+#include "Styling/SlateStyle.h"
+
+class FEditorTutorialStyle :
+    public FSlateStyleSet
+{
+public:
+	virtual const FName& GetStyleSetName() const override;
+
+	static const FEditorTutorialStyle& Get();
+	static void Shutdown();
+
+	~FEditorTutorialStyle();
+
+private:
+	FEditorTutorialStyle();
+
+	static FName StyleName;
+	static TUniquePtr<FEditorTutorialStyle> Instance;
+};
+
diff --git a/Engine/Plugins/Editor/GuidedTutorials/Source/IntroTutorials/Private/IntroTutorials.cpp b/Engine/Plugins/Editor/GuidedTutorials/Source/IntroTutorials/Private/IntroTutorials.cpp
--- a/Engine/Plugins/Editor/GuidedTutorials/Source/IntroTutorials/Private/IntroTutorials.cpp
+++ b/Engine/Plugins/Editor/GuidedTutorials/Source/IntroTutorials/Private/IntroTutorials.cpp
@@ -18,6 +18,7 @@
 
 #include "EditorTutorial.h"
 #include "EditorTutorialSettings.h"
+#include "EditorTutorialStyle.h"
 #include "TutorialStateSettings.h"
 #include "TutorialSettings.h"
 #include "ISettingsModule.h"
@@ -85,6 +86,8 @@ TSharedRef<FExtender> FIntroTutorials::AddSummonBlueprintTutorialsMenuExtender(c
 
 void FIntroTutorials::StartupModule()
 {
+	FEditorTutorialStyle::Get();
+
 	// This code can run with content commandlets. Slate is not initialized with commandlets and the below code will fail.
 	const bool bCommandlet = IsRunningCommandlet();
 	const bool bUnattended = FApp::IsUnattended();
@@ -219,6 +222,8 @@ void FIntroTutorials::ShutdownModule()
 	}
 
 	FGlobalTabmanager::Get()->UnregisterNomadTabSpawner(TEXT("TutorialsBrowser"));
+
+	FEditorTutorialStyle::Shutdown();
 }
 
 void FIntroTutorials::RegisterSummonTutorialsMenuEntries()
@@ -230,7 +235,7 @@ void FIntroTutorials::RegisterSummonTutorialsMenuEntries()
 		"Tutorials",
 		LOCTEXT("TutorialsMenuEntryTitle", "Tutorials"),
 		LOCTEXT("TutorialsMenuEntryToolTip", "Opens up introductory tutorials covering the basics of using the Unreal Engine 4 Editor."),
-		FSlateIcon(FEditorStyle::GetStyleSetName(), "LevelEditor.Tutorials"),
+		FSlateIcon(FEditorTutorialStyle::Get().GetStyleSetName(), "Tutorials.MenuIcon"),
 		FUIAction(FExecuteAction::CreateRaw(this, &FIntroTutorials::SummonTutorialHome))
 	));
 }
@@ -241,7 +246,7 @@ void FIntroTutorials::AddSummonBlueprintTutorialsMenuExtension(FMenuBuilder& Men
 	MenuBuilder.AddMenuEntry(
 		LOCTEXT("BlueprintMenuEntryTitle", "Blueprint Overview"),
 		LOCTEXT("BlueprintMenuEntryToolTip", "Opens up an introductory overview of Blueprints."),
-		FSlateIcon(FEditorStyle::GetStyleSetName(), "LevelEditor.Tutorials"),
+		FSlateIcon(FEditorTutorialStyle::Get().GetStyleSetName(), "Tutorials.MenuIcon"),
 		FUIAction(FExecuteAction::CreateRaw(this, &FIntroTutorials::SummonBlueprintTutorialHome, PrimaryObject, true)));
 
 	if(PrimaryObject != nullptr)
@@ -254,7 +259,7 @@ void FIntroTutorials::AddSummonBlueprintTutorialsMenuExtension(FMenuBuilder& Men
 			MenuBuilder.AddMenuEntry(
 				FText::Format(LOCTEXT("BlueprintTutorialsMenuEntryTitle", "{0} Tutorial"), Enum->GetDisplayNameTextByValue(BP->BlueprintType)),
 				LOCTEXT("BlueprintTutorialsMenuEntryToolTip", "Opens up an introductory tutorial covering this particular part of the Blueprint editor."),
-				FSlateIcon(FEditorStyle::GetStyleSetName(), "LevelEditor.Tutorials"),
+				FSlateIcon(FEditorTutorialStyle::Get().GetStyleSetName(), "Tutorials.MenuIcon"),
 				FUIAction(FExecuteAction::CreateRaw(this, &FIntroTutorials::SummonBlueprintTutorialHome, PrimaryObject, false)));
 		}
 	}
diff --git a/Engine/Plugins/Editor/GuidedTutorials/Source/IntroTutorials/Private/STutorialButton.cpp b/Engine/Plugins/Editor/GuidedTutorials/Source/IntroTutorials/Private/STutorialButton.cpp
--- a/Engine/Plugins/Editor/GuidedTutorials/Source/IntroTutorials/Private/STutorialButton.cpp
+++ b/Engine/Plugins/Editor/GuidedTutorials/Source/IntroTutorials/Private/STutorialButton.cpp
@@ -16,6 +16,7 @@
 #include "IIntroTutorials.h"
 #include "IntroTutorials.h"
 #include "EditorTutorialSettings.h"
+#include "EditorTutorialStyle.h"
 #include "TutorialStateSettings.h"
 #include "Misc/EngineBuildSettings.h"
 #include "AssetRegistryModule.h"
@@ -54,7 +55,7 @@ void STutorialButton::Construct(const FArguments& InArgs)
 	[
 		SNew(SButton)
 		.AddMetaData<FTagMetaData>(*FString::Printf(TEXT("%s.TutorialLaunchButton"), *Context.ToString()))
-		.ButtonStyle(FEditorStyle::Get(), "TutorialLaunch.Button")
+		.ButtonStyle(FEditorTutorialStyle::Get(), "TutorialLaunch.Button")
 		.ToolTipText(this, &STutorialButton::GetButtonToolTip)
 		.OnClicked(this, &STutorialButton::HandleButtonClicked)
 		.ContentPadding(0.0f)
@@ -118,8 +119,8 @@ int32 STutorialButton::OnPaint(const FPaintArgs& Args, const FGeometry& Allotted
 		float PulseFactor1 = 0.0f;
 		GetAnimationValues(PulseAnimation.GetLerp(), AlphaFactor0, PulseFactor0, AlphaFactor1, PulseFactor1);
 
-		const FSlateBrush* PulseBrush = FEditorStyle::Get().GetBrush(TEXT("TutorialLaunch.Circle"));
-		const FLinearColor PulseColor = FEditorStyle::Get().GetColor(TEXT("TutorialLaunch.Circle.Color"));
+		const FSlateBrush* PulseBrush = FEditorTutorialStyle::Get().GetBrush(TEXT("TutorialLaunch.Circle"));
+		const FLinearColor PulseColor = FEditorTutorialStyle::Get().GetColor(TEXT("TutorialLaunch.Circle.Color"));
 
 		// We should be clipped by the window size, not our containing widget, as we want to draw outside the widget
 		const FVector2D WindowSize = OutDrawElements.GetPaintWindow()->GetSizeInScreen();
diff --git a/Engine/Plugins/Editor/GuidedTutorials/Source/IntroTutorials/Private/STutorialContent.cpp b/Engine/Plugins/Editor/GuidedTutorials/Source/IntroTutorials/Private/STutorialContent.cpp
--- a/Engine/Plugins/Editor/GuidedTutorials/Source/IntroTutorials/Private/STutorialContent.cpp
+++ b/Engine/Plugins/Editor/GuidedTutorials/Source/IntroTutorials/Private/STutorialContent.cpp
@@ -1,6 +1,7 @@
 // Copyright Epic Games, Inc. All Rights Reserved.
 
 #include "STutorialContent.h"
+#include "EditorTutorialStyle.h"
 #include "Rendering/DrawElements.h"
 #include "Modules/ModuleManager.h"
 #include "Widgets/SBoxPanel.h"
@@ -108,9 +109,9 @@ void STutorialContent::Construct(const FArguments& InArgs, UEditorTutorial* InTu
 				// Add more padding if the content is to be displayed centrally (i.e. not on a widget)
 				.Padding(bIsStandalone ? TutorialConstants::BorderSizeStandalone : TutorialConstants::BorderSize)
 				.Visibility(EVisibility::SelfHitTestInvisible)
-				.BorderImage(FEditorStyle::GetBrush("Tutorials.Border"))
+				.BorderImage(FEditorTutorialStyle::Get().GetBrush("Tutorials.Border"))
 				.BorderBackgroundColor(this, &STutorialContent::GetBackgroundColor)
-				.ForegroundColor(FCoreStyle::Get().GetSlateColor("InvertedForeground"))
+				.ForegroundColor(FAppStyle::Get().GetSlateColor("InvertedForeground"))
 				[
 					SNew(SFxWidget)
 					.RenderScale(this, &STutorialContent::GetInverseAnimatedZoom)
@@ -146,7 +147,7 @@ void STutorialContent::Construct(const FArguments& InArgs, UEditorTutorial* InTu
 					SNew(SComboButton)
 					.ToolTipText(LOCTEXT("MoreOptionsTooltip", "More Options"))
 					.Visibility(this, &STutorialContent::GetMenuButtonVisibility)
-					.ButtonStyle(&FEditorStyle::Get().GetWidgetStyle<FButtonStyle>("Tutorials.Content.Button"))
+					.ButtonStyle(&FEditorTutorialStyle::Get().GetWidgetStyle<FButtonStyle>("Tutorials.Content.Button"))
 					.ContentPadding(0.0f)
 					.OnGetMenuContent(FOnGetContent::CreateSP(this, &STutorialContent::HandleGetMenuContent))
 				]
@@ -159,11 +160,11 @@ void STutorialContent::Construct(const FArguments& InArgs, UEditorTutorial* InTu
 					.ToolTipText(LOCTEXT("QuitStandaloneTooltip", "Close this Message"))
 					.OnClicked(this, &STutorialContent::OnCloseButtonClicked)
 					.Visibility(this, &STutorialContent::GetCloseButtonVisibility)
-					.ButtonStyle(&FEditorStyle::Get().GetWidgetStyle<FButtonStyle>("Tutorials.Content.Button"))
+					.ButtonStyle(&FEditorTutorialStyle::Get().GetWidgetStyle<FButtonStyle>("Tutorials.Content.Button"))
 					.ContentPadding(0.0f)
 					[
 						SNew(SImage)
-						.Image(FEditorStyle::GetBrush("Symbols.X"))
+						.Image(FAppStyle::Get().GetBrush("Symbols.X"))
 						.ColorAndOpacity(FLinearColor(0.0f, 0.0f, 0.0f, 1.0f))
 					]
 				]
@@ -177,7 +178,7 @@ void STutorialContent::Construct(const FArguments& InArgs, UEditorTutorial* InTu
 				.ToolTipText(this, &STutorialContent::GetBackButtonTooltip)
 				.OnClicked(this, &STutorialContent::HandleBackButtonClicked)
 				.Visibility(this, &STutorialContent::GetBackButtonVisibility)
-				.ButtonStyle(&FEditorStyle::Get().GetWidgetStyle<FButtonStyle>("Tutorials.Content.NavigationButtonWrapper"))
+				.ButtonStyle(&FEditorTutorialStyle::Get().GetWidgetStyle<FButtonStyle>("Tutorials.Content.NavigationButtonWrapper"))
 				.ContentPadding(0.0f)
 				[
 					SNew(SBox)
@@ -207,7 +208,7 @@ void STutorialContent::Construct(const FArguments& InArgs, UEditorTutorial* InTu
 				.ToolTipText(this, &STutorialContent::GetNextButtonTooltip)
 				.OnClicked(this, &STutorialContent::HandleNextClicked)
 				.Visibility(this, &STutorialContent::GetMenuButtonVisibility)
-				.ButtonStyle(&FEditorStyle::Get().GetWidgetStyle<FButtonStyle>("Tutorials.Content.NavigationButtonWrapper"))
+				.ButtonStyle(&FEditorTutorialStyle::Get().GetWidgetStyle<FButtonStyle>("Tutorials.Content.NavigationButtonWrapper"))
 				.ContentPadding(0.0f)
 				[
 					SNew(SBox)
@@ -224,7 +225,7 @@ void STutorialContent::Construct(const FArguments& InArgs, UEditorTutorial* InTu
 							[
 								SNew(STextBlock)
 								.Text(this, &STutorialContent::GetNextButtonLabel)
-								.TextStyle(FEditorStyle::Get(), "Tutorials.Content.NavigationText")
+								.TextStyle(FEditorTutorialStyle::Get(), "Tutorials.Content.NavigationText")
 								.ColorAndOpacity(FLinearColor::White)
 							]
 							+SHorizontalBox::Slot()
@@ -292,8 +293,8 @@ int32 STutorialContent::OnPaint(const FPaintArgs& Args, const FGeometry& Allotte
 		FLinearColor BorderTint;
 		GetAnimationValues(AlphaFactor, PulseFactor, ShadowTint, BorderTint);
 
-		const FSlateBrush* ShadowBrush = FEditorStyle::Get().GetBrush(TEXT("Tutorials.Shadow"));
-		const FSlateBrush* BorderBrush = FEditorStyle::Get().GetBrush(TEXT("Tutorials.Border"));
+		const FSlateBrush* ShadowBrush = FEditorTutorialStyle::Get().GetBrush(TEXT("Tutorials.Shadow"));
+		const FSlateBrush* BorderBrush = FEditorTutorialStyle::Get().GetBrush(TEXT("Tutorials.Border"));
 					
 		const FGeometry& WidgetGeometry = CachedGeometry;
 		const FVector2D WindowSize = OutDrawElements.GetPaintWindow()->GetSizeInScreen();
@@ -344,7 +345,7 @@ static TSharedRef<SWidget> GetStageTitle(const FExcerpt& InExcerpt, int32 InCurr
 	{
 		return SNew(STextBlock)
 			.Text(FText::FromString(*VariableValue))
-			.TextStyle(FEditorStyle::Get(), "Tutorials.CurrentExcerpt");
+			.TextStyle(FEditorTutorialStyle::Get(), "Tutorials.CurrentExcerpt");
 	}
 
 	// Then try 'StageTitle<StageNum>'
@@ -354,7 +355,7 @@ static TSharedRef<SWidget> GetStageTitle(const FExcerpt& InExcerpt, int32 InCurr
 	{
 		return SNew(STextBlock)
 			.Text(FText::FromString(*VariableValue))
-			.TextStyle(FEditorStyle::Get(), "Tutorials.CurrentExcerpt");
+			.TextStyle(FEditorTutorialStyle::Get(), "Tutorials.CurrentExcerpt");
 	}
 
 	return SNullWidget::NullWidget;
@@ -384,9 +385,9 @@ TSharedRef<SWidget> STutorialContent::GenerateContentWidget(const FTutorialConte
 				.Visibility(EVisibility::SelfHitTestInvisible)
 				.AutoWrapText(bAutoWrapText)
 				.Text(InContent.Text)
-				.TextStyle(FEditorStyle::Get(), "Tutorials.Content")
+				.TextStyle(FEditorTutorialStyle::Get(), "Tutorials.Content")
 				.HighlightText(InHighlightText)
-				.HighlightColor(FEditorStyle::Get().GetColor("Tutorials.Browser.HighlightTextColor"));
+				.HighlightColor(FEditorTutorialStyle::Get().GetColor("Tutorials.Browser.HighlightTextColor"));
 
 			if(!bAutoWrapText)
 			{
@@ -429,8 +430,8 @@ TSharedRef<SWidget> STutorialContent::GenerateContentWidget(const FTutorialConte
 
 			TSharedRef<SRichTextBlock> TextBlock = SNew(SRichTextBlock)
 					.Visibility(EVisibility::SelfHitTestInvisible)
-					.TextStyle(FEditorStyle::Get(), "Tutorials.Content.Text")
-					.DecoratorStyleSet(&FEditorStyle::Get())
+					.TextStyle(FEditorTutorialStyle::Get(), "Tutorials.Content.Text")
+					.DecoratorStyleSet(&FAppStyle::Get())
 					.Decorators(Decorators)
 					.Text(InContent.Text)
 					.AutoWrapText(bAutoWrapText)
@@ -589,7 +590,7 @@ FSlateColor STutorialContent::GetBackgroundColor() const
 {
 	// note cant use IsHovered() here because our widget is SelfHitTestInvisible
 	const FVector2D CursorPos = FSlateApplication::Get().GetCursorPos();
-	return CachedContentGeometry.IsUnderLocation(CursorPos) ? FEditorStyle::Get().GetColor("Tutorials.Content.Color.Hovered") : FEditorStyle::Get().GetColor("Tutorials.Content.Color");
+	return CachedContentGeometry.IsUnderLocation(CursorPos) ? FEditorTutorialStyle::Get().GetColor("Tutorials.Content.Color.Hovered") : FEditorTutorialStyle::Get().GetColor("Tutorials.Content.Color");
 }
 
 float STutorialContent::GetAnimatedZoom() const
@@ -754,11 +755,11 @@ const FSlateBrush* STutorialContent::GetNextButtonBrush() const
 {
 	if(IsNextEnabled.Get())
 	{
-		return FEditorStyle::GetBrush("Tutorials.Navigation.NextButton");
+		return FEditorTutorialStyle::Get().GetBrush("Tutorials.Navigation.NextButton");
 	}
 	else
 	{
-		return FEditorStyle::GetBrush("Tutorials.Navigation.HomeButton");
+		return FEditorTutorialStyle::Get().GetBrush("Tutorials.Navigation.HomeButton");
 	}
 }
 
@@ -795,16 +796,16 @@ FText STutorialContent::GetNextButtonLabel() const
 
 const FSlateBrush* STutorialContent::GetNextButtonBorder() const
 {
-	return NextButton->IsHovered() ? &FEditorStyle::Get().GetWidgetStyle<FButtonStyle>("Tutorials.Content.NavigationButton").Hovered : &FEditorStyle::Get().GetWidgetStyle<FButtonStyle>("Tutorials.Content.NavigationButton").Normal;
+	return NextButton->IsHovered() ? &FEditorTutorialStyle::Get().GetWidgetStyle<FButtonStyle>("Tutorials.Content.NavigationButton").Hovered : &FEditorTutorialStyle::Get().GetWidgetStyle<FButtonStyle>("Tutorials.Content.NavigationButton").Normal;
 }
 
 const FSlateBrush* STutorialContent::GetBackButtonBrush() const
 {
 	if (IsBackEnabled.Get())
 	{
-		return FEditorStyle::GetBrush("Tutorials.Navigation.BackButton");
+		return FEditorTutorialStyle::Get().GetBrush("Tutorials.Navigation.BackButton");
 	}
-	return FEditorStyle::GetDefaultBrush();
+	return FAppStyle::Get().GetDefaultBrush();
 }
 
 EVisibility STutorialContent::GetBackButtonVisibility() const
@@ -840,7 +841,7 @@ FText STutorialContent::GetBackButtonLabel() const
 
 const FSlateBrush* STutorialContent::GetBackButtonBorder() const
 {
-	return BackButton->IsHovered() ? &FEditorStyle::Get().GetWidgetStyle<FButtonStyle>("Tutorials.Content.NavigationBackButton").Hovered : &FEditorStyle::Get().GetWidgetStyle<FButtonStyle>("Tutorials.Content.NavigationBackButton").Normal;
+	return BackButton->IsHovered() ? &FEditorTutorialStyle::Get().GetWidgetStyle<FButtonStyle>("Tutorials.Content.NavigationBackButton").Hovered : &FEditorTutorialStyle::Get().GetWidgetStyle<FButtonStyle>("Tutorials.Content.NavigationBackButton").Normal;
 }
 
 FReply STutorialContent::OnMouseButtonDown( const FGeometry& MyGeometry, const FPointerEvent& MouseEvent )
diff --git a/Engine/Plugins/Editor/GuidedTutorials/Source/IntroTutorials/Private/STutorialEditableText.cpp b/Engine/Plugins/Editor/GuidedTutorials/Source/IntroTutorials/Private/STutorialEditableText.cpp
--- a/Engine/Plugins/Editor/GuidedTutorials/Source/IntroTutorials/Private/STutorialEditableText.cpp
+++ b/Engine/Plugins/Editor/GuidedTutorials/Source/IntroTutorials/Private/STutorialEditableText.cpp
@@ -1,6 +1,7 @@
 // Copyright Epic Games, Inc. All Rights Reserved.
 
 #include "STutorialEditableText.h"
+#include "EditorTutorialStyle.h"
 #include "Misc/PackageName.h"
 #include "SlateOptMacros.h"
 #include "Framework/Text/SlateImageRun.h"
@@ -40,7 +41,7 @@ void STutorialEditableText::Construct(const FArguments& InArgs)
 
 	TSharedRef<FRichTextLayoutMarshaller> RichTextMarshaller = FRichTextLayoutMarshaller::Create(
 		TArray<TSharedRef<ITextDecorator>>(), 
-		&FEditorStyle::Get()
+		&FAppStyle::Get()
 		);
 
 
@@ -64,7 +65,7 @@ void STutorialEditableText::Construct(const FArguments& InArgs)
 		.Padding(FMargin(0.0f, 0.0f, 0.0f, 0.0f))
 		[
 			SAssignNew(RichEditableTextBox, SMultiLineEditableTextBox)
-			.Font(FEditorStyle::Get().GetWidgetStyle<FTextBlockStyle>("Tutorials.Content.Text").Font)
+			.Font(FEditorTutorialStyle::Get().GetWidgetStyle<FTextBlockStyle>("Tutorials.Content.Text").Font)
 			.Text(InArgs._Text)
 			.OnTextChanged(this, &STutorialEditableText::HandleRichEditableTextChanged)
 			.OnTextCommitted(this, &STutorialEditableText::HandleRichEditableTextCommitted)
@@ -82,7 +83,7 @@ void STutorialEditableText::Construct(const FArguments& InArgs)
 		[
 			SNew(SBorder)
 			.Visibility(this, &STutorialEditableText::GetToolbarVisibility)
-			.BorderImage(FEditorStyle::Get().GetBrush("TutorialEditableText.RoundedBackground"))
+			.BorderImage(FEditorTutorialStyle::Get().GetBrush("TutorialEditableText.RoundedBackground"))
 			.Padding(FMargin(4))
 			[
 				SNew(SHorizontalBox)
@@ -91,7 +92,7 @@ void STutorialEditableText::Construct(const FArguments& InArgs)
 				.AutoWidth()
 				[
 					SAssignNew(FontComboBox, SComboBox<TSharedPtr<FTextStyleAndName>>)
-					.ComboBoxStyle(FEditorStyle::Get(), "TutorialEditableText.Toolbar.ComboBox")
+					.ComboBoxStyle(FEditorTutorialStyle::Get(), "TutorialEditableText.Toolbar.ComboBox")
 					.OptionsSource(&StylesAndNames)
 					.OnSelectionChanged(this, &STutorialEditableText::OnActiveStyleChanged)
 					.OnGenerateWidget(this, &STutorialEditableText::GenerateStyleComboEntry)
@@ -114,14 +115,14 @@ void STutorialEditableText::Construct(const FArguments& InArgs)
 				[
 					SAssignNew(HyperlinkComboButton, SComboButton)
 					.ToolTipText(LOCTEXT("HyperlinkButtonTooltip", "Insert or Edit Hyperlink"))
-					.ComboButtonStyle(FEditorStyle::Get(), "TutorialEditableText.Toolbar.ComboButton")
+					.ComboButtonStyle(FEditorTutorialStyle::Get(), "TutorialEditableText.Toolbar.ComboButton")
 					.OnComboBoxOpened(this, &STutorialEditableText::HandleHyperlinkComboOpened)
 					.IsEnabled(this, &STutorialEditableText::IsHyperlinkComboEnabled)
 					.ContentPadding(1.0f)
 					.ButtonContent()
 					[
 						SNew(SImage)
-						.Image(FEditorStyle::Get().GetBrush("TutorialEditableText.Toolbar.HyperlinkImage"))
+						.Image(FEditorTutorialStyle::Get().GetBrush("TutorialEditableText.Toolbar.HyperlinkImage"))
 					]
 					.MenuContent()
 					[
@@ -134,7 +135,7 @@ void STutorialEditableText::Construct(const FArguments& InArgs)
 						.Padding(FMargin(2.0f))
 						[
 							SNew(STextBlock)
-							.TextStyle(FEditorStyle::Get(), "TutorialEditableText.Toolbar.Text")
+							.TextStyle(FEditorTutorialStyle::Get(), "TutorialEditableText.Toolbar.Text")
 							.Text(LOCTEXT("HyperlinkNameLabel", "Name:"))
 						]
 						+SGridPanel::Slot(1, 0)
@@ -144,7 +145,7 @@ void STutorialEditableText::Construct(const FArguments& InArgs)
 							.WidthOverride(300)
 							[
 								SAssignNew(HyperlinkNameTextBlock, STextBlock)
-								.TextStyle(FEditorStyle::Get(), "TutorialEditableText.Toolbar.Text")
+								.TextStyle(FEditorTutorialStyle::Get(), "TutorialEditableText.Toolbar.Text")
 							]
 						]
 
@@ -154,7 +155,7 @@ void STutorialEditableText::Construct(const FArguments& InArgs)
 						.Padding(FMargin(2.0f))
 						[
 							SNew(STextBlock)
-							.TextStyle(FEditorStyle::Get(), "TutorialEditableText.Toolbar.Text")
+							.TextStyle(FEditorTutorialStyle::Get(), "TutorialEditableText.Toolbar.Text")
 							.Text(LOCTEXT("HyperlinkURLLabel", "URL:"))
 						]
 						+SGridPanel::Slot(1, 1)
@@ -173,7 +174,7 @@ void STutorialEditableText::Construct(const FArguments& InArgs)
 						.Padding(FMargin(2.0f))
 						[
 							SNew(STextBlock)
-							.TextStyle(FEditorStyle::Get(), "TutorialEditableText.Toolbar.Text")
+							.TextStyle(FEditorTutorialStyle::Get(), "TutorialEditableText.Toolbar.Text")
 							.Text(LOCTEXT("HyperlinkTypeLabel", "Type:"))
 						]
 
@@ -189,7 +190,7 @@ void STutorialEditableText::Construct(const FArguments& InArgs)
 							[
 								SNew(SComboBox<TSharedPtr<FHyperlinkTypeDesc>>)
 								.OptionsSource(&FTutorialText::GetHyperLinkDescs())
-								.ComboBoxStyle(FEditorStyle::Get(), "TutorialEditableText.Toolbar.ComboBox")
+								.ComboBoxStyle(FEditorTutorialStyle::Get(), "TutorialEditableText.Toolbar.ComboBox")
 								.OnSelectionChanged(this, &STutorialEditableText::OnActiveHyperlinkChanged)
 								.OnGenerateWidget(this, &STutorialEditableText::GenerateHyperlinkComboEntry)
 								.ContentPadding(0.0f)
@@ -201,7 +202,7 @@ void STutorialEditableText::Construct(const FArguments& InArgs)
 									.MinDesiredWidth(100.0f)
 									[
 										SNew(STextBlock)
-										.TextStyle(FEditorStyle::Get(), "TutorialEditableText.Toolbar.Text")
+										.TextStyle(FEditorTutorialStyle::Get(), "TutorialEditableText.Toolbar.Text")
 										.Text(this, &STutorialEditableText::GetActiveHyperlinkName)
 										.ToolTipText(this, &STutorialEditableText::GetActiveHyperlinkTooltip)
 									]
@@ -219,7 +220,7 @@ void STutorialEditableText::Construct(const FArguments& InArgs)
 								.OnCheckStateChanged(this, &STutorialEditableText::HandleOpenAssetCheckStateChanged)
 								[
 									SNew(STextBlock)
-									.TextStyle(FEditorStyle::Get(), "TutorialEditableText.Toolbar.Text")
+									.TextStyle(FEditorTutorialStyle::Get(), "TutorialEditableText.Toolbar.Text")
 									.Text(LOCTEXT("OpenAssetLabel", "Open Asset"))
 								]
 							]
@@ -239,11 +240,11 @@ void STutorialEditableText::Construct(const FArguments& InArgs)
 							.HAlign(HAlign_Right)
 							[
 								SNew(SButton)
-								.ButtonStyle(FEditorStyle::Get(), "TutorialEditableText.Toolbar.Button")
+								.ButtonStyle(FEditorTutorialStyle::Get(), "TutorialEditableText.Toolbar.Button")
 								.OnClicked(this, &STutorialEditableText::HandleInsertHyperLinkClicked)
 								[
 									SNew(STextBlock)
-									.TextStyle(FEditorStyle::Get(), "TutorialEditableText.Toolbar.Text")
+									.TextStyle(FEditorTutorialStyle::Get(), "TutorialEditableText.Toolbar.Text")
 									.Text(this, &STutorialEditableText::GetHyperlinkButtonText)
 								]
 							]
@@ -257,13 +258,13 @@ void STutorialEditableText::Construct(const FArguments& InArgs)
 				[
 					SNew(SButton)
 					.ToolTipText(LOCTEXT("ImageButtonTooltip", "Insert Image"))
-					.ButtonStyle(FEditorStyle::Get(), "TutorialEditableText.Toolbar.Button")
+					.ButtonStyle(FEditorTutorialStyle::Get(), "TutorialEditableText.Toolbar.Button")
 					.OnClicked(this, &STutorialEditableText::HandleImageButtonClicked)
 					.ContentPadding(1.0f)
 					.Content()
 					[
 						SNew(SImage)
-						.Image(FEditorStyle::Get().GetBrush("TutorialEditableText.Toolbar.ImageImage"))
+						.Image(FEditorTutorialStyle::Get().GetBrush("TutorialEditableText.Toolbar.ImageImage"))
 					]
 				]
 			]
@@ -390,12 +391,12 @@ void STutorialEditableText::OnActiveStyleChanged(TSharedPtr<FTextStyleAndName> N
 TSharedRef<SWidget> STutorialEditableText::GenerateStyleComboEntry(TSharedPtr<FTextStyleAndName> SourceEntry)
 {
 	return SNew(SBorder)
-		.BorderImage( FCoreStyle::Get().GetBrush( "NoBorder" ) )
-		.ForegroundColor(FCoreStyle::Get().GetSlateColor("InvertedForeground"))
+		.BorderImage( FAppStyle::Get().GetBrush( "NoBorder" ) )
+		.ForegroundColor(FAppStyle::Get().GetSlateColor("InvertedForeground"))
 		[
 			SNew(STextBlock)
 			.Text(SourceEntry->DisplayName)
-			.TextStyle(&FEditorStyle::Get().GetWidgetStyle<FTextBlockStyle>(SourceEntry->Style))
+			.TextStyle(&FEditorTutorialStyle::Get().GetWidgetStyle<FTextBlockStyle>(SourceEntry->Style))
 		];
 }
 
@@ -494,7 +495,7 @@ FReply STutorialEditableText::HandleInsertHyperLinkClicked()
 		TSharedRef<FSlateHyperlinkRun> HyperlinkRun = FSlateHyperlinkRun::Create(
 			RunInfo, 
 			MakeShareable(new FString(Name.ToString())), 
-			FEditorStyle::Get().GetWidgetStyle<FHyperlinkStyle>(FName(TEXT("Tutorials.Content.Hyperlink"))), 
+			FEditorTutorialStyle::Get().GetWidgetStyle<FHyperlinkStyle>(FName(TEXT("Tutorials.Content.Hyperlink"))),
 			CurrentHyperlinkType->OnClickedDelegate,
 			CurrentHyperlinkType->TooltipDelegate,
 			CurrentHyperlinkType->TooltipTextDelegate
@@ -527,13 +528,13 @@ void STutorialEditableText::OnActiveHyperlinkChanged(TSharedPtr<FHyperlinkTypeDe
 TSharedRef<SWidget> STutorialEditableText::GenerateHyperlinkComboEntry(TSharedPtr<FHyperlinkTypeDesc> SourceEntry)
 {
 	return SNew(SBorder)
-		.BorderImage( FCoreStyle::Get().GetBrush( "NoBorder" ) )
-		.ForegroundColor(FCoreStyle::Get().GetSlateColor("InvertedForeground"))
+		.BorderImage( FAppStyle::Get().GetBrush( "NoBorder" ) )
+		.ForegroundColor(FAppStyle::Get().GetSlateColor("InvertedForeground"))
 		[
 			SNew(STextBlock)
 			.Text(SourceEntry->Text)
 			.ToolTipText(SourceEntry->TooltipText)
-			.TextStyle(&FEditorStyle::Get().GetWidgetStyle<FTextBlockStyle>("TutorialEditableText.Toolbar.Text"))
+			.TextStyle(&FEditorTutorialStyle::Get().GetWidgetStyle<FTextBlockStyle>("TutorialEditableText.Toolbar.Text"))
 		];
 }
 
diff --git a/Engine/Plugins/Editor/GuidedTutorials/Source/IntroTutorials/Private/STutorialLoading.cpp b/Engine/Plugins/Editor/GuidedTutorials/Source/IntroTutorials/Private/STutorialLoading.cpp
--- a/Engine/Plugins/Editor/GuidedTutorials/Source/IntroTutorials/Private/STutorialLoading.cpp
+++ b/Engine/Plugins/Editor/GuidedTutorials/Source/IntroTutorials/Private/STutorialLoading.cpp
@@ -23,7 +23,7 @@ void STutorialLoading::Construct(const FArguments& InArgs)
 			SNew(SBorder)
 			.VAlign(VAlign_Center)
 			.HAlign(HAlign_Center)
-			.BorderImage(FEditorStyle::GetBrush("ToolPanel.GroupBorder"))
+			.BorderImage(FAppStyle::Get().GetBrush("ToolPanel.GroupBorder"))
 			[
 				SNew(SVerticalBox)
 				+ SVerticalBox::Slot()
diff --git a/Engine/Plugins/Editor/GuidedTutorials/Source/IntroTutorials/Private/STutorialNavigation.cpp b/Engine/Plugins/Editor/GuidedTutorials/Source/IntroTutorials/Private/STutorialNavigation.cpp
--- a/Engine/Plugins/Editor/GuidedTutorials/Source/IntroTutorials/Private/STutorialNavigation.cpp
+++ b/Engine/Plugins/Editor/GuidedTutorials/Source/IntroTutorials/Private/STutorialNavigation.cpp
@@ -1,6 +1,7 @@
 // Copyright Epic Games, Inc. All Rights Reserved.
 
 #include "STutorialNavigation.h"
+#include "EditorTutorialStyle.h"
 #include "Widgets/SBoxPanel.h"
 #include "Styling/SlateTypes.h"
 #include "Widgets/Layout/SBorder.h"
@@ -23,7 +24,7 @@ void STutorialNavigation::Construct(const FArguments& InArgs)
 	[
 		SNew(SBorder)
 		.Padding(24.0f)
-		.BorderImage(FEditorStyle::GetBrush("Tutorials.Border"))
+		.BorderImage(FEditorTutorialStyle::Get().GetBrush("Tutorials.Border"))
 		[
 			SNew(SVerticalBox)
 			+SVerticalBox::Slot()
@@ -37,11 +38,11 @@ void STutorialNavigation::Construct(const FArguments& InArgs)
 					SNew(SButton)
 					.OnClicked(this, &STutorialNavigation::OnBackButtonClicked)
 					.IsEnabled(InArgs._IsBackEnabled)
-					.ButtonStyle(&FEditorStyle::Get().GetWidgetStyle<FButtonStyle>("Tutorials.Content.Button"))
+					.ButtonStyle(&FEditorTutorialStyle::Get().GetWidgetStyle<FButtonStyle>("Tutorials.Content.Button"))
 					.Content()
 					[
 						SNew(SImage)
-						.Image(FEditorStyle::GetBrush("Tutorials.Back"))
+						.Image(FEditorTutorialStyle::Get().GetBrush("Tutorials.Back"))
 						.ColorAndOpacity(this, &STutorialNavigation::GetBackButtonColor)
 					]
 				]
@@ -52,11 +53,11 @@ void STutorialNavigation::Construct(const FArguments& InArgs)
 					SNew(SButton)
 					.OnClicked(this, &STutorialNavigation::OnHomeButtonClicked)
 					.IsEnabled(InArgs._IsHomeEnabled)
-					.ButtonStyle(&FEditorStyle::Get().GetWidgetStyle<FButtonStyle>("Tutorials.Content.Button"))
+					.ButtonStyle(&FEditorTutorialStyle::Get().GetWidgetStyle<FButtonStyle>("Tutorials.Content.Button"))
 					.Content()
 					[
 						SNew(SImage)
-						.Image(FEditorStyle::GetBrush("Tutorials.Home"))
+						.Image(FEditorTutorialStyle::Get().GetBrush("Tutorials.Home"))
 						.ColorAndOpacity(this, &STutorialNavigation::GetHomeButtonColor)
 					]
 				]
@@ -67,11 +68,11 @@ void STutorialNavigation::Construct(const FArguments& InArgs)
 					SNew(SButton)
 					.OnClicked(this, &STutorialNavigation::OnNextButtonClicked)
 					.IsEnabled(InArgs._IsNextEnabled)
-					.ButtonStyle(&FEditorStyle::Get().GetWidgetStyle<FButtonStyle>("Tutorials.Content.Button"))
+					.ButtonStyle(&FEditorTutorialStyle::Get().GetWidgetStyle<FButtonStyle>("Tutorials.Content.Button"))
 					.Content()
 					[
 						SNew(SImage)
-						.Image(FEditorStyle::GetBrush("Tutorials.Next"))
+						.Image(FEditorTutorialStyle::Get().GetBrush("Tutorials.Next"))
 						.ColorAndOpacity(this, &STutorialNavigation::GetNextButtonColor)
 					]
 				]
diff --git a/Engine/Plugins/Editor/GuidedTutorials/Source/IntroTutorials/Private/STutorialOverlay.cpp b/Engine/Plugins/Editor/GuidedTutorials/Source/IntroTutorials/Private/STutorialOverlay.cpp
--- a/Engine/Plugins/Editor/GuidedTutorials/Source/IntroTutorials/Private/STutorialOverlay.cpp
+++ b/Engine/Plugins/Editor/GuidedTutorials/Source/IntroTutorials/Private/STutorialOverlay.cpp
@@ -168,7 +168,7 @@ int32 STutorialOverlay::TraverseWidgets(TSharedRef<SWidget> InWidget, const FGeo
 		if(WidgetNameToHighlight != NAME_None )
 		{
 			const FLinearColor Color = bIsPicking && bShouldHighlight ? FLinearColor::Green : FLinearColor::White;
-			FSlateDrawElement::MakeBox(OutDrawElements, LayerId++, InGeometry.ToPaintGeometry(), FCoreStyle::Get().GetBrush(TEXT("Debug.Border")), ESlateDrawEffect::None, Color);
+			FSlateDrawElement::MakeBox(OutDrawElements, LayerId++, InGeometry.ToPaintGeometry(), FAppStyle::Get().GetBrush(TEXT("Debug.Border")), ESlateDrawEffect::None, Color);
 		}	
 	}
 
diff --git a/Engine/Plugins/Editor/GuidedTutorials/Source/IntroTutorials/Private/STutorialsBrowser.cpp b/Engine/Plugins/Editor/GuidedTutorials/Source/IntroTutorials/Private/STutorialsBrowser.cpp
--- a/Engine/Plugins/Editor/GuidedTutorials/Source/IntroTutorials/Private/STutorialsBrowser.cpp
+++ b/Engine/Plugins/Editor/GuidedTutorials/Source/IntroTutorials/Private/STutorialsBrowser.cpp
@@ -17,6 +17,7 @@
 #include "STutorialContent.h"
 #include "TutorialSettings.h"
 #include "EditorTutorialSettings.h"
+#include "EditorTutorialStyle.h"
 #include "TutorialStateSettings.h"
 #include "ARFilter.h"
 #include "AssetRegistryModule.h"
@@ -85,13 +86,13 @@ public:
 		{
 			if(Category.Icon.Len() > 0)
 			{
-				SlateBrush = FEditorStyle::Get().GetBrush(FName(*Category.Icon));
+				SlateBrush = FAppStyle::Get().GetBrush(FName(*Category.Icon));
 			}
 		}
 
 		if(SlateBrush == nullptr)
 		{
-			SlateBrush = FEditorStyle::Get().GetBrush("Tutorials.Browser.DefaultTutorialIcon");
+			SlateBrush = FEditorTutorialStyle::Get().GetBrush("Tutorials.Browser.DefaultTutorialIcon");
 		}
 	}
 
@@ -108,7 +109,7 @@ public:
 			[
 				SNew(SButton)
 				.OnClicked(this, &FTutorialListEntry_Category::OnClicked)
-				.ButtonStyle(&FEditorStyle::Get().GetWidgetStyle<FButtonStyle>("Tutorials.Browser.Button"))
+				.ButtonStyle(&FEditorTutorialStyle::Get().GetWidgetStyle<FButtonStyle>("Tutorials.Browser.Button"))
 				.ForegroundColor(FSlateColor::UseForeground())
 				.Content()
 				[
@@ -139,7 +140,7 @@ public:
 							SNew(SImage)
 							.ToolTipText(LOCTEXT("CompletedCategoryCheckToolTip", "This category has been completed"))
 							.Visibility(this, &FTutorialListEntry_Category::GetCompletedVisibility)
-							.Image(FEditorStyle::Get().GetBrush("Tutorials.Browser.Completed"))
+							.Image(FEditorTutorialStyle::Get().GetBrush("Tutorials.Browser.Completed"))
 						]
 					]
 					+ SHorizontalBox::Slot()
@@ -152,10 +153,10 @@ public:
 						[
 							SNew(STextBlock)
 							.Text(GetTitleText())
-							.TextStyle(&FEditorStyle::Get().GetWidgetStyle<FTextBlockStyle>("Tutorials.Browser.SummaryHeader"))
+							.TextStyle(&FEditorTutorialStyle::Get().GetWidgetStyle<FTextBlockStyle>("Tutorials.Browser.SummaryHeader"))
 							.HighlightText(HighlightText)
-							.HighlightColor(FEditorStyle::Get().GetColor("Tutorials.Browser.HighlightTextColor"))
-							.HighlightShape(FEditorStyle::Get().GetBrush("TextBlock.HighlightShape"))
+							.HighlightColor(FEditorTutorialStyle::Get().GetColor("Tutorials.Browser.HighlightTextColor"))
+							.HighlightShape(FAppStyle::Get().GetBrush("TextBlock.HighlightShape"))
 						]
 						+SVerticalBox::Slot()
 						.FillHeight(1.0f)
@@ -163,10 +164,10 @@ public:
 							SNew(STextBlock)
 							.AutoWrapText(true)
 							.Text(Category.Description)
-							.TextStyle(&FEditorStyle::Get().GetWidgetStyle<FTextBlockStyle>("Tutorials.Browser.SummaryText"))
+							.TextStyle(&FEditorTutorialStyle::Get().GetWidgetStyle<FTextBlockStyle>("Tutorials.Browser.SummaryText"))
 							.HighlightText(HighlightText)
-							.HighlightColor(FEditorStyle::Get().GetColor("Tutorials.Browser.HighlightTextColor"))
-							.HighlightShape(FEditorStyle::Get().GetBrush("TextBlock.HighlightShape"))
+							.HighlightColor(FEditorTutorialStyle::Get().GetColor("Tutorials.Browser.HighlightTextColor"))
+							.HighlightShape(FAppStyle::Get().GetBrush("TextBlock.HighlightShape"))
 						]
 					]
 					+SHorizontalBox::Slot()
@@ -176,7 +177,7 @@ public:
 					[
 						SNew(SImage)
 						.Visibility(this, &FTutorialListEntry_Category::OnGetArrowVisibility)
-						.Image(FEditorStyle::Get().GetBrush("Tutorials.Browser.CategoryArrow"))
+						.Image(FEditorTutorialStyle::Get().GetBrush("Tutorials.Browser.CategoryArrow"))
 					]
 				]
 			]
@@ -295,12 +296,12 @@ public:
 		}	
 		else if(Tutorial->Icon.Len() > 0)
 		{
-			SlateBrush = FEditorStyle::Get().GetBrush(FName(*Tutorial->Icon));
+			SlateBrush = FAppStyle::Get().GetBrush(FName(*Tutorial->Icon));
 		}
 		
 		if(SlateBrush == nullptr)
 		{
-			SlateBrush = FEditorStyle::Get().GetBrush("Tutorials.Browser.DefaultTutorialIcon");
+			SlateBrush = FEditorTutorialStyle::Get().GetBrush("Tutorials.Browser.DefaultTutorialIcon");
 		}
 	}
 
@@ -319,7 +320,7 @@ public:
 			[
 				SAssignNew(LaunchButton, SButton)
 				.OnClicked(this, &FTutorialListEntry_Tutorial::OnClicked, false)
-				.ButtonStyle(&FEditorStyle::Get().GetWidgetStyle<FButtonStyle>("Tutorials.Browser.Button"))
+				.ButtonStyle(&FEditorTutorialStyle::Get().GetWidgetStyle<FButtonStyle>("Tutorials.Browser.Button"))
 				.ForegroundColor(FSlateColor::UseForeground())
 				.Content()
 				[
@@ -354,7 +355,7 @@ public:
 								SNew(SImage)
 								.ToolTipText(LOCTEXT("CompletedTutorialCheckToolTip", "This tutorial has been completed"))
 								.Visibility(this, &FTutorialListEntry_Tutorial::GetCompletedVisibility)
-								.Image(FEditorStyle::Get().GetBrush("Tutorials.Browser.Completed"))
+								.Image(FEditorTutorialStyle::Get().GetBrush("Tutorials.Browser.Completed"))
 							]
 						]
 						+SHorizontalBox::Slot()
@@ -372,10 +373,10 @@ public:
 								[
 									SNew(STextBlock)
 									.Text(GetTitleText())
-									.TextStyle(&FEditorStyle::Get().GetWidgetStyle<FTextBlockStyle>("Tutorials.Browser.SummaryHeader"))
+									.TextStyle(&FEditorTutorialStyle::Get().GetWidgetStyle<FTextBlockStyle>("Tutorials.Browser.SummaryHeader"))
 									.HighlightText(HighlightText)
-									.HighlightColor(FEditorStyle::Get().GetColor("Tutorials.Browser.HighlightTextColor"))
-									.HighlightShape(FEditorStyle::Get().GetBrush("TextBlock.HighlightShape"))
+									.HighlightColor(FEditorTutorialStyle::Get().GetColor("Tutorials.Browser.HighlightTextColor"))
+									.HighlightShape(FAppStyle::Get().GetBrush("TextBlock.HighlightShape"))
 								]
 								+SHorizontalBox::Slot()
 								.AutoWidth()
@@ -385,11 +386,11 @@ public:
 									.ToolTipText(LOCTEXT("RestartButtonToolTip", "Start this tutorial from the beginning"))
 									.Visibility(this, &FTutorialListEntry_Tutorial::GetRestartVisibility)
 									.OnClicked(this, &FTutorialListEntry_Tutorial::OnClicked, true)
-									.ButtonStyle(&FEditorStyle::Get().GetWidgetStyle<FButtonStyle>("Tutorials.Browser.Button"))
+									.ButtonStyle(&FEditorTutorialStyle::Get().GetWidgetStyle<FButtonStyle>("Tutorials.Browser.Button"))
 									.Content()
 									[
 										SNew(SImage)
-										.Image(FEditorStyle::GetBrush("Tutorials.Browser.RestartButton"))
+										.Image(FEditorTutorialStyle::Get().GetBrush("Tutorials.Browser.RestartButton"))
 									]
 								]
 							]
@@ -563,7 +564,7 @@ void STutorialsBrowser::Construct(const FArguments& InArgs)
 		.AutoHeight()
 		[
 			SNew(SBorder)
-			.BorderImage(FEditorStyle::Get().GetBrush("ToolPanel.GroupBorder"))
+			.BorderImage(FAppStyle::Get().GetBrush("ToolPanel.GroupBorder"))
 			.Padding(5.0f)
 			[
 				SNew(SHorizontalBox)
@@ -574,12 +575,12 @@ void STutorialsBrowser::Construct(const FArguments& InArgs)
 					SNew(SButton)
 					.OnClicked(this, &STutorialsBrowser::OnBackButtonClicked)
 					.IsEnabled(this, &STutorialsBrowser::IsBackButtonEnabled)
-					.ButtonStyle(&FEditorStyle::Get().GetWidgetStyle<FButtonStyle>("Tutorials.Browser.BackButton"))
+					.ButtonStyle(&FEditorTutorialStyle::Get().GetWidgetStyle<FButtonStyle>("Tutorials.Browser.BackButton"))
 					.ForegroundColor(FSlateColor::UseForeground())
 					.Content()
 					[
 						SNew(SImage)
-						.Image(FEditorStyle::GetBrush("Tutorials.Browser.BackButton.Image"))
+						.Image(FEditorTutorialStyle::Get().GetBrush("Tutorials.Browser.BackButton.Image"))
 					]
 				]
 				+SHorizontalBox::Slot()
@@ -594,8 +595,8 @@ void STutorialsBrowser::Construct(const FArguments& InArgs)
 					[
 						SAssignNew(BreadcrumbTrail, SBreadcrumbTrail<TSharedPtr<ITutorialListEntry>>)
 						.ButtonContentPadding(FMargin(1.0f, 1.0f))
-						.DelimiterImage(FEditorStyle::GetBrush("Tutorials.Browser.Breadcrumb"))
-						.TextStyle(FEditorStyle::Get(), "Tutorials.Browser.PathText")
+						.DelimiterImage(FEditorTutorialStyle::Get().GetBrush("Tutorials.Browser.Breadcrumb"))
+						.TextStyle(FEditorTutorialStyle::Get(), "Tutorials.Browser.PathText")
 						.ShowLeadingDelimiter( true )
 						.OnCrumbClicked(this, &STutorialsBrowser::OnBreadcrumbClicked)
 					]
diff --git a/Engine/Plugins/Editor/GuidedTutorials/Source/IntroTutorials/Private/TutorialHyperlinkDecorator.cpp b/Engine/Plugins/Editor/GuidedTutorials/Source/IntroTutorials/Private/TutorialHyperlinkDecorator.cpp
--- a/Engine/Plugins/Editor/GuidedTutorials/Source/IntroTutorials/Private/TutorialHyperlinkDecorator.cpp
+++ b/Engine/Plugins/Editor/GuidedTutorials/Source/IntroTutorials/Private/TutorialHyperlinkDecorator.cpp
@@ -3,6 +3,7 @@
 #include "TutorialHyperlinkDecorator.h"
 #include "Styling/CoreStyle.h"
 #include "TutorialHyperlinkRun.h"
+#include "EditorTutorialStyle.h"
 
 TSharedRef< FTutorialHyperlinkDecorator > FTutorialHyperlinkDecorator::Create( FString Id, const FSlateHyperlinkRun::FOnClick& NavigateDelegate, const FSlateHyperlinkRun::FOnGetTooltipText& InToolTipTextDelegate, const FSlateHyperlinkRun::FOnGenerateTooltip& InToolTipDelegate )
 {
@@ -32,7 +33,7 @@ TSharedRef< ISlateRun > FTutorialHyperlinkDecorator::Create(const TSharedRef<cla
 
 	if ( !Style->HasWidgetStyle<FHyperlinkStyle>( FName( *StyleName ) ) )
 	{
-		Style = &FCoreStyle::Get();
+		Style = &FEditorTutorialStyle::Get();
 	}
 
 	FRunInfo RunInfo( RunParseResult.Name );
diff --git a/Engine/Plugins/Editor/GuidedTutorials/Source/IntroTutorials/Private/TutorialHyperlinkRun.cpp b/Engine/Plugins/Editor/GuidedTutorials/Source/IntroTutorials/Private/TutorialHyperlinkRun.cpp
--- a/Engine/Plugins/Editor/GuidedTutorials/Source/IntroTutorials/Private/TutorialHyperlinkRun.cpp
+++ b/Engine/Plugins/Editor/GuidedTutorials/Source/IntroTutorials/Private/TutorialHyperlinkRun.cpp
@@ -7,6 +7,7 @@
 #include "Widgets/Images/SImage.h"
 #include "EditorStyleSet.h"
 #include "Widgets/Input/SRichTextHyperlink.h"
+#include "EditorTutorialStyle.h"
 
 TSharedRef< FTutorialHyperlinkRun > FTutorialHyperlinkRun::Create( const FRunInfo& InRunInfo, const TSharedRef< const FString >& InText, const FHyperlinkStyle& InStyle, FOnClick NavigateDelegate, FOnGenerateTooltip InTooltipDelegate, FOnGetTooltipText InTooltipTextDelegate )
 {
@@ -90,7 +91,7 @@ TSharedRef< ILayoutBlock > FTutorialHyperlinkRun::CreateBlock( int32 StartIndex,
 FTutorialHyperlinkRun::FTutorialHyperlinkRun( const FRunInfo& InRunInfo, const TSharedRef< const FString >& InText, const FHyperlinkStyle& InStyle, FOnClick InNavigateDelegate, FOnGenerateTooltip InTooltipDelegate, FOnGetTooltipText InTooltipTextDelegate ) 
 	: FSlateHyperlinkRun(InRunInfo, InText, InStyle, InNavigateDelegate, InTooltipDelegate, InTooltipTextDelegate)
 	, bIsBrowserLink(CheckIsBrowserLink())
-	, LinkBrush(FEditorStyle::Get().GetBrush("Tutorials.Content.ExternalLink"))
+	, LinkBrush(FEditorTutorialStyle::Get().GetBrush("Tutorials.Content.ExternalLink"))
 {
 
 }
@@ -98,7 +99,7 @@ FTutorialHyperlinkRun::FTutorialHyperlinkRun( const FRunInfo& InRunInfo, const T
 FTutorialHyperlinkRun::FTutorialHyperlinkRun( const FRunInfo& InRunInfo, const TSharedRef< const FString >& InText, const FHyperlinkStyle& InStyle, FOnClick InNavigateDelegate, FOnGenerateTooltip InTooltipDelegate, FOnGetTooltipText InTooltipTextDelegate, const FTextRange& InRange ) 
 	: FSlateHyperlinkRun(InRunInfo, InText, InStyle, InNavigateDelegate, InTooltipDelegate, InTooltipTextDelegate,InRange)
 	, bIsBrowserLink(CheckIsBrowserLink())
-	, LinkBrush(FEditorStyle::Get().GetBrush("Tutorials.Content.ExternalLink"))
+	, LinkBrush(FEditorTutorialStyle::Get().GetBrush("Tutorials.Content.ExternalLink"))
 {
 
 }
diff --git a/Engine/Plugins/Editor/GuidedTutorials/Source/IntroTutorials/Private/TutorialStructCustomization.cpp b/Engine/Plugins/Editor/GuidedTutorials/Source/IntroTutorials/Private/TutorialStructCustomization.cpp
--- a/Engine/Plugins/Editor/GuidedTutorials/Source/IntroTutorials/Private/TutorialStructCustomization.cpp
+++ b/Engine/Plugins/Editor/GuidedTutorials/Source/IntroTutorials/Private/TutorialStructCustomization.cpp
@@ -455,14 +455,14 @@ public:
 			.VAlign(VAlign_Center)
 			[
 				SNew(SButton)
-				.ButtonStyle( FEditorStyle::Get(), "HoverHintOnly" )
+				.ButtonStyle( FAppStyle::Get(), "HoverHintOnly" )
 				.OnClicked( this, &SWidgetPicker::OnClicked )
 				.ContentPadding(4.0f)
 				.ForegroundColor( FSlateColor::UseForeground() )
 				.IsFocusable(false)
 				[
 					SNew( SImage )
-					.Image( FEditorStyle::GetBrush("Icons.EyeDropper") )
+					.Image(FAppStyle::Get().GetBrush("Icons.EyeDropper") )
 					.ColorAndOpacity( FSlateColor::UseForeground() )
 				]
 			]
diff --git a/Engine/Plugins/Editor/GuidedTutorials/Source/IntroTutorials/Private/TutorialText.h b/Engine/Plugins/Editor/GuidedTutorials/Source/IntroTutorials/Private/TutorialText.h
--- a/Engine/Plugins/Editor/GuidedTutorials/Source/IntroTutorials/Private/TutorialText.h
+++ b/Engine/Plugins/Editor/GuidedTutorials/Source/IntroTutorials/Private/TutorialText.h
@@ -7,6 +7,7 @@
 #include "Framework/Text/IRun.h"
 #include "Framework/Text/SlateHyperlinkRun.h"
 #include "EditorStyleSet.h"
+#include "EditorTutorialStyle.h"
 
 namespace EHyperlinkType
 {
@@ -153,12 +154,12 @@ struct FTextStyleAndName
 
 	FTextBlockStyle CreateTextBlockStyle() const
 	{
-		return FEditorStyle::Get().GetWidgetStyle<FTextBlockStyle>(Style);
+		return FEditorTutorialStyle::Get().GetWidgetStyle<FTextBlockStyle>(Style);
 	}
 
 	static FTextBlockStyle CreateTextBlockStyle(const FRunInfo& InRunInfo)
 	{
-		return FEditorStyle::Get().GetWidgetStyle<FTextBlockStyle>(GetStyleFromRunInfo(InRunInfo));
+		return FEditorTutorialStyle::Get().GetWidgetStyle<FTextBlockStyle>(GetStyleFromRunInfo(InRunInfo));
 	}
 
 	/** The style identifier */
diff --git a/Engine/Source/Editor/EditorStyle/Private/StarshipStyle.cpp b/Engine/Source/Editor/EditorStyle/Private/StarshipStyle.cpp
--- a/Engine/Source/Editor/EditorStyle/Private/StarshipStyle.cpp
+++ b/Engine/Source/Editor/EditorStyle/Private/StarshipStyle.cpp
@@ -2643,18 +2643,22 @@ void FStarshipEditorStyle::FStyle::SetupTutorialStyles()
 	const FTextBlockStyle DocumentationText = FTextBlockStyle(NormalText)
 		.SetColorAndOpacity( FLinearColor::Black )
 		.SetFont(DEFAULT_FONT( "Regular", 11 ));
+	Set("Documentation.Text", FTextBlockStyle(DocumentationText));
 
 	const FTextBlockStyle DocumentationHyperlinkText = FTextBlockStyle(DocumentationText)
 		.SetColorAndOpacity( HyperlinkColor );
+	Set("Documentation.Hyperlink.Text", FTextBlockStyle(DocumentationHyperlinkText));
 
 	const FTextBlockStyle DocumentationHeaderText = FTextBlockStyle(NormalText)
 		.SetColorAndOpacity( FLinearColor::Black )
 		.SetFont(DEFAULT_FONT("Black", 32));
+	Set("Documentation.Header.Text", FTextBlockStyle(DocumentationHeaderText));
 
 	const FButtonStyle DocumentationHyperlinkButton = FButtonStyle()
 			.SetNormal(BORDER_BRUSH( "Old/HyperlinkDotted", FMargin(0,0,0,3/16.0f), HyperlinkColor ) )
 			.SetPressed(FSlateNoResource())
 			.SetHovered(BORDER_BRUSH( "Old/HyperlinkUnderline", FMargin(0,0,0,3/16.0f), HyperlinkColor ) );
+	Set("Documentation.Hyperlink.Button", FButtonStyle(DocumentationHyperlinkButton));
 
 	// Documentation
 	{
@@ -2684,312 +2688,6 @@ void FStarshipEditorStyle::FStyle::SetupTutorialStyles()
 	{
 		Set("Documentation.ToolTip.Background", new BOX_BRUSH("Tutorials/TutorialContentBackground", FMargin(4 / 16.0f)));
 	}
-
-	// Tutorials
-	{
-		const FLinearColor TutorialButtonColor = FLinearColor(0.15f, 0.15f, 0.15f, 1.0f);
-		const FLinearColor TutorialSelectionColor = FLinearColor(0.19f, 0.33f, 0.72f);
-		const FLinearColor TutorialNavigationButtonColor = FLinearColor(0.0f, 0.59f, 0.14f, 1.0f);
-		const FLinearColor TutorialNavigationButtonHoverColor = FLinearColor(0.2f, 0.79f, 0.34f, 1.0f);
-		const FLinearColor TutorialNavigationBackButtonColor = TutorialNavigationButtonColor;
-		const FLinearColor TutorialNavigationBackButtonHoverColor = TutorialNavigationButtonHoverColor;
-
-		const FTextBlockStyle TutorialText = FTextBlockStyle(DocumentationText)
-			.SetColorAndOpacity(FLinearColor::Black)
-			.SetHighlightColor(TutorialSelectionColor);
-
-		const FTextBlockStyle TutorialHeaderText = FTextBlockStyle(DocumentationHeaderText)
-			.SetColorAndOpacity(FLinearColor::Black)
-			.SetHighlightColor(TutorialSelectionColor);
-
-		Set("Tutorials.Border", new BOX_BRUSH( "Tutorials/OverlayFrame", FMargin(18.0f/64.0f), FLinearColor(1.0f, 1.0f, 1.0f, 1.0f) ) );
-		Set("Tutorials.Shadow", new BOX_BRUSH("Tutorials/TutorialShadow", FVector2D(256.0f, 256.0f), FMargin(114.0f / 256.0f)));
-
-		const FTextBlockStyle TutorialBrowserText = FTextBlockStyle(TutorialText)
-			.SetColorAndOpacity(FSlateColor::UseForeground())
-			.SetHighlightColor(TutorialSelectionColor);
-
-		Set( "Tutorials.Browser.Text", TutorialBrowserText );
-
-		Set( "Tutorials.Browser.WelcomeHeader", FTextBlockStyle(TutorialBrowserText)
-			.SetFontSize(20));
-
-		Set( "Tutorials.Browser.SummaryHeader", FTextBlockStyle(TutorialBrowserText)
-			.SetFontSize(16));
-
-		Set( "Tutorials.Browser.SummaryText", FTextBlockStyle(TutorialBrowserText)
-			.SetFontSize(10));
-
-		Set( "Tutorials.Browser.HighlightTextColor", TutorialSelectionColor );
-
-		Set( "Tutorials.Browser.Button", FButtonStyle()
-			.SetNormal( BOX_BRUSH( "Common/ButtonHoverHint", FMargin(4/16.0f), FLinearColor(0.05f,0.05f,0.05f,1) ) )
-			.SetHovered( BOX_BRUSH( "Common/ButtonHoverHint", FMargin(4/16.0f), FLinearColor(0.07f,0.07f,0.07f,1) ) )
-			.SetPressed( BOX_BRUSH( "Common/ButtonHoverHint", FMargin(4/16.0f), FLinearColor(0.08f,0.08f,0.08f,1) ) )
-			.SetNormalPadding( FMargin(0,0,0,1))
-			.SetPressedPadding( FMargin(0,1,0,0)));
-
-		Set( "Tutorials.Browser.BackButton", FButtonStyle()
-			.SetNormal( BOX_BRUSH( "Common/ButtonHoverHint", FMargin(4/16.0f), FLinearColor(1.0f,1.0f,1.0f,0.0f) ) )
-			.SetHovered( BOX_BRUSH( "Common/ButtonHoverHint", FMargin(4/16.0f), FLinearColor(1.0f,1.0f,1.0f,0.05f) ) )
-			.SetPressed( BOX_BRUSH( "Common/ButtonHoverHint", FMargin(4/16.0f), FLinearColor(1.0f,1.0f,1.0f,0.05f) ) )
-			.SetNormalPadding( FMargin(0,0,0,1))
-			.SetPressedPadding( FMargin(0,1,0,0)));
-
-		Set( "Tutorials.Content.Button", FButtonStyle()
-			.SetNormal( BOX_BRUSH( "Common/ButtonHoverHint", FMargin(4/16.0f), FLinearColor(0,0,0,0) ) )
-			.SetHovered( BOX_BRUSH( "Common/ButtonHoverHint", FMargin(4/16.0f), FLinearColor(1,1,1,1) ) )
-			.SetPressed( BOX_BRUSH( "Common/ButtonHoverHint", FMargin(4/16.0f), FLinearColor(1,1,1,1) ) )
-			.SetNormalPadding( FMargin(0,0,0,1))
-			.SetPressedPadding( FMargin(0,1,0,0)));
-
-		Set( "Tutorials.Content.NavigationButtonWrapper", FButtonStyle()
-			.SetNormal( FSlateNoResource() )
-			.SetHovered( FSlateNoResource() )
-			.SetPressed( FSlateNoResource() )
-			.SetNormalPadding( FMargin(0,0,0,1))
-			.SetPressedPadding( FMargin(0,1,0,0)));
-
-		Set( "Tutorials.Content.NavigationButton", FButtonStyle()
-			.SetNormal( BOX_BRUSH( "Common/ButtonHoverHint", FMargin(4/16.0f), TutorialNavigationButtonColor ) )
-			.SetHovered( BOX_BRUSH( "Common/ButtonHoverHint", FMargin(4/16.0f), TutorialNavigationButtonHoverColor ) )
-			.SetPressed( BOX_BRUSH( "Common/ButtonHoverHint", FMargin(4/16.0f), TutorialNavigationButtonHoverColor ) )
-			.SetNormalPadding( FMargin(0,0,0,1))
-			.SetPressedPadding( FMargin(0,1,0,0)));
-
-		Set("Tutorials.Content.NavigationBackButton", FButtonStyle()
-			.SetNormal(BOX_BRUSH("Common/ButtonHoverHint", FMargin(4 / 16.0f), TutorialNavigationBackButtonColor))
-			.SetHovered(BOX_BRUSH("Common/ButtonHoverHint", FMargin(4 / 16.0f), TutorialNavigationBackButtonHoverColor))
-			.SetPressed(BOX_BRUSH("Common/ButtonHoverHint", FMargin(4 / 16.0f), TutorialNavigationBackButtonHoverColor))
-			.SetNormalPadding(FMargin(0, 0, 0, 1))
-			.SetPressedPadding(FMargin(0, 1, 0, 0)));
-
-		Set( "Tutorials.Content.NavigationText", FTextBlockStyle(TutorialText));
-
-		Set( "Tutorials.Content.Color", FLinearColor(1.0f,1.0f,1.0f,0.9f) );
-		Set( "Tutorials.Content.Color.Hovered", FLinearColor(1.0f,1.0f,1.0f,1.0f) );
-
-		Set( "Tutorials.Browser.CategoryArrow", new IMAGE_BRUSH( "Tutorials/BrowserCategoryArrow", FVector2D(24.0f, 24.0f), FSlateColor::UseForeground() ) );
-		Set( "Tutorials.Browser.DefaultTutorialIcon", new IMAGE_BRUSH( "Tutorials/DefaultTutorialIcon_40x", FVector2D(40.0f, 40.0f), FLinearColor::White ) );
-		Set( "Tutorials.Browser.DefaultCategoryIcon", new IMAGE_BRUSH( "Tutorials/DefaultCategoryIcon_40x", FVector2D(40.0f, 40.0f), FLinearColor::White ) );
-
-		Set( "Tutorials.Browser.BackButton.Image", new IMAGE_BRUSH( "Tutorials/BrowserBack", FVector2D(32.0f, 32.0f), FLinearColor(1.0f, 1.0f, 1.0f, 1.0f) ) );
-		Set( "Tutorials.Browser.PlayButton.Image", new IMAGE_BRUSH( "Tutorials/BrowserPlay", FVector2D(32.0f, 32.0f), FLinearColor(1.0f, 1.0f, 1.0f, 1.0f) ) );
-		Set( "Tutorials.Browser.RestartButton", new IMAGE_BRUSH( "Tutorials/BrowserRestart", FVector2D(16.0f, 16.0f), FLinearColor(1.0f, 1.0f, 1.0f, 1.0f) ) );
-
-		Set( "Tutorials.Browser.Completed", new IMAGE_BRUSH( "Tutorials/TutorialCompleted", Icon32x32 ) );
-
-		Set( "Tutorials.Browser.Breadcrumb", new IMAGE_BRUSH( "Tutorials/Breadcrumb", Icon8x8, FLinearColor::White ) );
-		Set( "Tutorials.Browser.PathText", FTextBlockStyle(TutorialBrowserText)
-			.SetFontSize(9));
-
-		Set( "Tutorials.Navigation.Button", FButtonStyle()
-			.SetNormal( BOX_BRUSH( "Common/ButtonHoverHint", FMargin(4/16.0f), FLinearColor(0,0,0,0) ) )
-			.SetHovered( BOX_BRUSH( "Common/ButtonHoverHint", FMargin(4/16.0f), FLinearColor(0,0,0,0) ) )
-			.SetPressed( BOX_BRUSH( "Common/ButtonHoverHint", FMargin(4/16.0f), FLinearColor(0,0,0,0) ) )
-			.SetNormalPadding( FMargin(0,0,0,1))
-			.SetPressedPadding( FMargin(0,1,0,0)));
-
-		Set( "Tutorials.Navigation.NextButton", new IMAGE_BRUSH( "Tutorials/NavigationNext", Icon32x32 ) );
-		Set( "Tutorials.Navigation.HomeButton", new IMAGE_BRUSH( "Tutorials/NavigationHome", Icon32x32 ) );
-		Set( "Tutorials.Navigation.BackButton", new IMAGE_BRUSH( "Tutorials/NavigationBack", Icon32x32 ) );
-
-		Set( "Tutorials.WidgetContent", FTextBlockStyle(TutorialText)
-			.SetFontSize(10));
-
-		Set( "Tutorials.ButtonColor", TutorialButtonColor );
-		Set( "Tutorials.ButtonHighlightColor", TutorialSelectionColor );
-		Set( "Tutorials.ButtonDisabledColor", SelectionColor_Inactive );
-		Set( "Tutorials.ContentAreaBackground", new BOX_BRUSH( "Tutorials/TutorialContentBackground", FMargin(4/16.0f) ) );
-		Set( "Tutorials.HomeContentAreaBackground", new BOX_BRUSH( "Tutorials/TutorialHomeContentBackground", FMargin(4/16.0f) ) );
-		Set( "Tutorials.ContentAreaFrame", new BOX_BRUSH( "Tutorials/ContentAreaFrame", FMargin(26.0f/64.0f) ) );
-		Set( "Tutorials.CurrentExcerpt", new IMAGE_BRUSH( "Tutorials/CurrentExcerpt", FVector2D(24.0f, 24.0f), TutorialSelectionColor ) );
-		Set( "Tutorials.Home", new IMAGE_BRUSH( "Tutorials/HomeButton", FVector2D(32.0f, 32.0f) ) );
-		Set( "Tutorials.Back", new IMAGE_BRUSH("Tutorials/BackButton", FVector2D(24.0f, 24.0f) ) );
-		Set( "Tutorials.Next", new IMAGE_BRUSH("Tutorials/NextButton", FVector2D(24.0f, 24.0f) ) );
-
-		Set("Tutorials.PageHeader", FTextBlockStyle(TutorialHeaderText)
-			.SetFontSize(22));
-
-		Set("Tutorials.CurrentExcerpt", FTextBlockStyle(TutorialHeaderText)
-			.SetFontSize(16));
- 
-		Set("Tutorials.NavigationButtons", FTextBlockStyle(TutorialHeaderText)
-			.SetFontSize(16));
-
-		// UDN documentation styles
-		Set("Tutorials.Content", FTextBlockStyle(TutorialText)
-			.SetColorAndOpacity(FSlateColor::UseForeground()));
-		Set("Tutorials.Hyperlink.Text",  FTextBlockStyle(DocumentationHyperlinkText));
-		Set("Tutorials.NumberedContent", FTextBlockStyle(TutorialText));
-		Set("Tutorials.BoldContent",     FTextBlockStyle(TutorialText)
-			.SetTypefaceFontName(TEXT("Bold")));
-
-		Set("Tutorials.Header1", FTextBlockStyle(TutorialHeaderText)
-			.SetFontSize(32));
-
-		Set("Tutorials.Header2", FTextBlockStyle(TutorialHeaderText)
-			.SetFontSize(24));
-
-		Set( "Tutorials.Hyperlink.Button", FButtonStyle(DocumentationHyperlinkButton)
-			.SetNormal(BORDER_BRUSH( "Old/HyperlinkDotted", FMargin(0,0,0,3/16.0f), FLinearColor::Black))
-			.SetHovered(BORDER_BRUSH( "Old/HyperlinkUnderline", FMargin(0,0,0,3/16.0f), FLinearColor::Black)));
-
-		Set( "Tutorials.Separator", new BOX_BRUSH( "Common/Separator", 1/4.0f, FLinearColor::Black));
-
-		Set( "Tutorials.ProgressBar", FProgressBarStyle()
-			.SetBackgroundImage( BOX_BRUSH( "Common/ProgressBar_Background", FMargin(5.f/12.f) ) )
-			.SetFillImage( BOX_BRUSH( "Common/ProgressBar_NeutralFill", FMargin(5.f/12.f) ) )
-			.SetMarqueeImage( IMAGE_BRUSH( "Common/ProgressBar_Marquee", FVector2D(20,12), FLinearColor::White, ESlateBrushTileType::Horizontal ) )
-			);
-
-		// Default text styles
-		{
-			const FTextBlockStyle RichTextNormal = FTextBlockStyle()
-				.SetFont(DEFAULT_FONT("Regular", 11))
-				.SetColorAndOpacity(FSlateColor::UseForeground())
-				.SetShadowOffset(FVector2D::ZeroVector)
-				.SetShadowColorAndOpacity(FLinearColor::Black)
-				.SetHighlightColor(FLinearColor(0.02f, 0.3f, 0.0f))
-				.SetHighlightShape(BOX_BRUSH("Common/TextBlockHighlightShape", FMargin(3.f /8.f)));
-			Set( "Tutorials.Content.Text", RichTextNormal );
-
-			Set( "Tutorials.Content.TextBold", FTextBlockStyle(RichTextNormal)
-				.SetFont(DEFAULT_FONT("Bold", 11)));
-
-			Set( "Tutorials.Content.HeaderText1", FTextBlockStyle(RichTextNormal)
-				.SetFontSize(20));
-
-			Set( "Tutorials.Content.HeaderText2", FTextBlockStyle(RichTextNormal)
-				.SetFontSize(16));
-
-			{
-				const FButtonStyle RichTextHyperlinkButton = FButtonStyle()
-					.SetNormal(BORDER_BRUSH("Old/HyperlinkDotted", FMargin(0,0,0,3/16.0f), FLinearColor::Blue ) )
-					.SetPressed(FSlateNoResource() )
-					.SetHovered(BORDER_BRUSH("Old/HyperlinkUnderline", FMargin(0,0,0,3/16.0f), FLinearColor::Blue ) );
-
-				const FTextBlockStyle RichTextHyperlinkText = FTextBlockStyle(RichTextNormal)
-					.SetColorAndOpacity(FLinearColor::Blue);
-
-				Set( "Tutorials.Content.HyperlinkText", RichTextHyperlinkText );
-
-				// legacy style
-				Set( "TutorialEditableText.Editor.HyperlinkText", RichTextHyperlinkText );
-
-				const FHyperlinkStyle RichTextHyperlink = FHyperlinkStyle()
-					.SetUnderlineStyle(RichTextHyperlinkButton)
-					.SetTextStyle(RichTextHyperlinkText)
-					.SetPadding(FMargin(0.0f));
-				Set( "Tutorials.Content.Hyperlink", RichTextHyperlink );
-
-				Set( "Tutorials.Content.ExternalLink", new IMAGE_BRUSH("Tutorials/ExternalLink", Icon16x16, FLinearColor::Blue));
-
-				// legacy style
-				Set( "TutorialEditableText.Editor.Hyperlink", RichTextHyperlink );
-			}
-		}
-
-		// Toolbar
-		{
-			const FLinearColor NormalColor(FColor(0xffeff3f3));
-			const FLinearColor SelectedColor(FColor(0xffdbe4d5));
-			const FLinearColor HoverColor(FColor(0xffdbe4e4));
-			const FLinearColor DisabledColor(FColor(0xaaaaaa));
-			const FLinearColor TextColor(FColor(0xff2c3e50));
-
-			Set("TutorialEditableText.RoundedBackground", new BOX_BRUSH("Common/RoundedSelection_16x", 4.0f / 16.0f, FLinearColor(FColor(0xffeff3f3))));
-
-			Set("TutorialEditableText.Toolbar.HyperlinkImage", new IMAGE_BRUSH("Tutorials/hyperlink", Icon16x16, TextColor));
-			Set("TutorialEditableText.Toolbar.ImageImage", new IMAGE_BRUSH("Tutorials/Image", Icon16x16, TextColor));
-
-			Set("TutorialEditableText.Toolbar.TextColor", TextColor);
-
-			Set("TutorialEditableText.Toolbar.Text", FTextBlockStyle(NormalText)
-				.SetFont(DEFAULT_FONT("Regular", 10))
-				.SetColorAndOpacity(TextColor)
-				);
-
-			Set("TutorialEditableText.Toolbar.BoldText", FTextBlockStyle(NormalText)
-				.SetFont(DEFAULT_FONT("Bold", 10))
-				.SetColorAndOpacity(TextColor)
-				);
-
-			Set("TutorialEditableText.Toolbar.ItalicText", FTextBlockStyle(NormalText)
-				.SetFont(DEFAULT_FONT("Italic", 10))
-				.SetColorAndOpacity(TextColor)
-				);
-
-			Set("TutorialEditableText.Toolbar.Checkbox", FCheckBoxStyle()
-				.SetCheckBoxType(ESlateCheckBoxType::CheckBox)
-				.SetUncheckedImage(IMAGE_BRUSH("Common/CheckBox", Icon16x16, FLinearColor::White))
-				.SetUncheckedHoveredImage(IMAGE_BRUSH("Common/CheckBox", Icon16x16, HoverColor))
-				.SetUncheckedPressedImage(IMAGE_BRUSH("Common/CheckBox_Hovered", Icon16x16, HoverColor))
-				.SetCheckedImage(IMAGE_BRUSH("Common/CheckBox_Checked_Hovered", Icon16x16, FLinearColor::White))
-				.SetCheckedHoveredImage(IMAGE_BRUSH("Common/CheckBox_Checked_Hovered", Icon16x16, HoverColor))
-				.SetCheckedPressedImage(IMAGE_BRUSH("Common/CheckBox_Checked", Icon16x16, HoverColor))
-				.SetUndeterminedImage(IMAGE_BRUSH("Common/CheckBox_Undetermined", Icon16x16, FLinearColor::White))
-				.SetUndeterminedHoveredImage(IMAGE_BRUSH("Common/CheckBox_Undetermined_Hovered", Icon16x16, HoverColor))
-				.SetUndeterminedPressedImage(IMAGE_BRUSH("Common/CheckBox_Undetermined_Hovered", Icon16x16, FLinearColor::White))
-				);
-
-			Set("TutorialEditableText.Toolbar.ToggleButtonCheckbox", FCheckBoxStyle()
-				.SetCheckBoxType(ESlateCheckBoxType::ToggleButton)
-				.SetUncheckedImage(BOX_BRUSH("Tutorials/FlatColorSquare", FVector2D(1.0f, 1.0f), FMargin(1), NormalColor))
-				.SetUncheckedHoveredImage(BOX_BRUSH("Tutorials/FlatColorSquare", FVector2D(1.0f, 1.0f), FMargin(1), HoverColor))
-				.SetUncheckedPressedImage(BOX_BRUSH("Tutorials/FlatColorSquare", FVector2D(1.0f, 1.0f), FMargin(1), HoverColor))
-				.SetCheckedImage(BOX_BRUSH("Tutorials/FlatColorSquare", FVector2D(1.0f, 1.0f), FMargin(1), SelectedColor))
-				.SetCheckedHoveredImage(BOX_BRUSH("Tutorials/FlatColorSquare", FVector2D(1.0f, 1.0f), FMargin(1), HoverColor))
-				.SetCheckedPressedImage(BOX_BRUSH("Tutorials/FlatColorSquare", FVector2D(1.0f, 1.0f), FMargin(1), HoverColor))
-				);
-
-			const FButtonStyle TutorialButton = FButtonStyle()
-				.SetNormal(BOX_BRUSH("Tutorials/FlatColorSquare", FVector2D(1.0f, 1.0f), FMargin(1), NormalColor))
-				.SetHovered(BOX_BRUSH("Tutorials/FlatColorSquare", FVector2D(1.0f, 1.0f), FMargin(1), HoverColor))
-				.SetPressed(BOX_BRUSH("Tutorials/FlatColorSquare", FVector2D(1.0f, 1.0f), FMargin(1), SelectedColor))
-				.SetNormalPadding(FMargin(2, 2, 2, 2))
-				.SetPressedPadding(FMargin(2, 3, 2, 1));
-			Set("TutorialEditableText.Toolbar.Button", TutorialButton);
-
-			const FComboButtonStyle ComboButton = FComboButtonStyle()
-				.SetButtonStyle(Button)
-				.SetDownArrowImage(IMAGE_BRUSH("Common/ComboArrow", Icon8x8))
-				.SetMenuBorderBrush(BOX_BRUSH("Tutorials/FlatColorSquare", FVector2D(1.0f, 1.0f), FMargin(1), NormalColor))
-				.SetMenuBorderPadding(FMargin(0.0f));
-			Set("TutorialEditableText.Toolbar.ComboButton", ComboButton);
-
-			{
-				const FButtonStyle ComboBoxButton = FButtonStyle()
-					.SetNormal(BOX_BRUSH("Tutorials/FlatColorSquare", FVector2D(1.0f, 1.0f), FMargin(1), FLinearColor::White))
-					.SetHovered(BOX_BRUSH("Tutorials/FlatColorSquare", FVector2D(1.0f, 1.0f), FMargin(1), FLinearColor::White))
-					.SetPressed(BOX_BRUSH("Tutorials/FlatColorSquare", FVector2D(1.0f, 1.0f), FMargin(1), FLinearColor::White))
-					.SetNormalPadding(FMargin(2, 2, 2, 2))
-					.SetPressedPadding(FMargin(2, 3, 2, 1));
-
-				const FComboButtonStyle ComboBoxComboButton = FComboButtonStyle(ComboButton)
-					.SetButtonStyle(ComboBoxButton)
-					.SetMenuBorderPadding(FMargin(1.0));
-
-				Set("TutorialEditableText.Toolbar.ComboBox", FComboBoxStyle()
-					.SetComboButtonStyle(ComboBoxComboButton)
-					);
-			}
-		}
-
-		// In-editor tutorial launch button
-		{
-			Set("TutorialLaunch.Button", FButtonStyle()
-				.SetNormalPadding(0)
-				.SetPressedPadding(0)
-				.SetNormal(IMAGE_BRUSH_SVG("Starship/Common/Tutorials", Icon16x16, FStyleColors::Foreground))
-				.SetHovered(IMAGE_BRUSH_SVG("Starship/Common/Tutorials", Icon16x16, FStyleColors::ForegroundHover))
-				.SetPressed(IMAGE_BRUSH_SVG("Starship/Common/Tutorials", Icon16x16, FStyleColors::ForegroundHover))
-				);
-
-			Set("TutorialLaunch.Circle", new IMAGE_BRUSH("Tutorials/Circle_128x", Icon128x128, FLinearColor::White));
-			Set("TutorialLaunch.Circle.Color", FLinearColor::Green);
-		}
-	}
 }
 
 void FStarshipEditorStyle::FStyle::SetupPropertyEditorStyles()
