Move EditorTutorial asset icon and thumbnail to plugin style.

diff --git a/Engine/Plugins/Editor/GuidedTutorials/Source/IntroTutorials/Private/EditorTutorialStyle.cpp b/Engine/Plugins/Editor/GuidedTutorials/Source/IntroTutorials/Private/EditorTutorialStyle.cpp
--- a/Engine/Plugins/Editor/GuidedTutorials/Source/IntroTutorials/Private/EditorTutorialStyle.cpp
+++ b/Engine/Plugins/Editor/GuidedTutorials/Source/IntroTutorials/Private/EditorTutorialStyle.cpp
@@ -340,6 +340,9 @@ FEditorTutorialStyle::FEditorTutorialStyle()
 	Set("TutorialLaunch.Circle", new IMAGE_BRUSH("Tutorials/Circle_128x", Icon128x128, FLinearColor::White));
 	Set("TutorialLaunch.Circle.Color", FLinearColor::Green);
 
+	Set("ClassIcon.EditorTutorial", new IMAGE_BRUSH_SVG("AssetIcons/EditorTutorial_16", Icon16x16));
+	Set("ClassThumbnail.EditorTutorial", new IMAGE_BRUSH_SVG("AssetIcons/EditorTutorial_64", Icon16x16));
+
 	FSlateStyleRegistry::RegisterSlateStyle(*this);
 }
 
diff --git a/Engine/Source/Editor/EditorStyle/Private/StarshipStyle.cpp b/Engine/Source/Editor/EditorStyle/Private/StarshipStyle.cpp
--- a/Engine/Source/Editor/EditorStyle/Private/StarshipStyle.cpp
+++ b/Engine/Source/Editor/EditorStyle/Private/StarshipStyle.cpp
@@ -5787,7 +5787,6 @@ void FStarshipEditorStyle::FStyle::SetupClassIconsAndThumbnails()
 			{TEXT("DialogueWave")},
 			{TEXT("DirectionalLight")},
 			{TEXT("DocumentationActor")},
-			{TEXT("EditorTutorial")},
 			{TEXT("EnvQuery")},
 			{TEXT("Emitter")},
 			{TEXT("EmptyActor")},
