Take DPI settings into account when placing editor tutorial bubbles.

diff --git a/Engine/Plugins/Editor/GuidedTutorials/Source/IntroTutorials/IntroTutorials.Build.cs b/Engine/Plugins/Editor/GuidedTutorials/Source/IntroTutorials/IntroTutorials.Build.cs
--- a/Engine/Plugins/Editor/GuidedTutorials/Source/IntroTutorials/IntroTutorials.Build.cs
+++ b/Engine/Plugins/Editor/GuidedTutorials/Source/IntroTutorials/IntroTutorials.Build.cs
@@ -18,7 +18,8 @@ namespace UnrealBuildTool.Rules
                     "Documentation",
 					"GraphEditor",
 					"BlueprintGraph",
-					"MessageLog"
+					"MessageLog",
+					"ApplicationCore"
 				}
 			);
 
diff --git a/Engine/Plugins/Editor/GuidedTutorials/Source/IntroTutorials/Private/STutorialContent.cpp b/Engine/Plugins/Editor/GuidedTutorials/Source/IntroTutorials/Private/STutorialContent.cpp
--- a/Engine/Plugins/Editor/GuidedTutorials/Source/IntroTutorials/Private/STutorialContent.cpp
+++ b/Engine/Plugins/Editor/GuidedTutorials/Source/IntroTutorials/Private/STutorialContent.cpp
@@ -32,6 +32,7 @@
 #include "EngineAnalytics.h"
 #include "AnalyticsEventAttribute.h"
 #include "Interfaces/IAnalyticsProvider.h"
+#include "HAL/PlatformApplicationMisc.h"
 
 #define LOCTEXT_NAMESPACE "STutorialContent"
 
@@ -505,7 +506,10 @@ FVector2D STutorialContent::GetPosition() const
 		YOffset += WidgetOffset.Get().Y;
 
 		// now build & clamp to area
-		FVector2D BaseOffset = FVector2D(CachedGeometry.AbsolutePosition.X + XOffset, CachedGeometry.AbsolutePosition.Y + YOffset);
+		float WidgetX = CachedGeometry.GetAbsolutePosition().X;
+		float WidgetY = CachedGeometry.GetAbsolutePosition().Y;
+		float DPIScale = FPlatformApplicationMisc::GetDPIScaleFactorAtPoint(WidgetX, WidgetY);
+		FVector2D BaseOffset = FVector2D(WidgetX / DPIScale + XOffset, WidgetY / DPIScale + YOffset);
 		BaseOffset.X = FMath::Clamp(BaseOffset.X, 0.0f, CachedWindowSize.X - ContentWidget->GetDesiredSize().X);
 		BaseOffset.Y = FMath::Clamp(BaseOffset.Y, 0.0f, CachedWindowSize.Y - ContentWidget->GetDesiredSize().Y);
 		return BaseOffset;
