Restore tutorial highlight styling to its UE4 appearance, copied from the legacy CoreStyle.

diff --git a/Engine/Plugins/Editor/GuidedTutorials/Source/IntroTutorials/Private/EditorTutorialStyle.cpp b/Engine/Plugins/Editor/GuidedTutorials/Source/IntroTutorials/Private/EditorTutorialStyle.cpp
--- a/Engine/Plugins/Editor/GuidedTutorials/Source/IntroTutorials/Private/EditorTutorialStyle.cpp
+++ b/Engine/Plugins/Editor/GuidedTutorials/Source/IntroTutorials/Private/EditorTutorialStyle.cpp
@@ -58,6 +58,8 @@ FEditorTutorialStyle::FEditorTutorialStyle()
 	Set("Tutorials.Border", new BOX_BRUSH("Tutorials/OverlayFrame", FMargin(18.0f / 64.0f), FLinearColor(1.0f, 1.0f, 1.0f, 1.0f)));
 	Set("Tutorials.Shadow", new BOX_BRUSH("Tutorials/TutorialShadow", FVector2D(256.0f, 256.0f), FMargin(114.0f / 256.0f)));
 
+	Set("Tutorials.Highlight.Border", new BOX_BRUSH("Tutorials/TutorialBorder", FVector2D(64.0f, 64.0f), FMargin(25.0f / 64.0f)));
+
 	const FTextBlockStyle TutorialBrowserText = FTextBlockStyle(TutorialText)
 		.SetColorAndOpacity(FSlateColor::UseForeground())
 		.SetHighlightColor(TutorialSelectionColor);
diff --git a/Engine/Plugins/Editor/GuidedTutorials/Source/IntroTutorials/Private/STutorialContent.cpp b/Engine/Plugins/Editor/GuidedTutorials/Source/IntroTutorials/Private/STutorialContent.cpp
--- a/Engine/Plugins/Editor/GuidedTutorials/Source/IntroTutorials/Private/STutorialContent.cpp
+++ b/Engine/Plugins/Editor/GuidedTutorials/Source/IntroTutorials/Private/STutorialContent.cpp
@@ -295,7 +295,7 @@ int32 STutorialContent::OnPaint(const FPaintArgs& Args, const FGeometry& Allotte
 		GetAnimationValues(AlphaFactor, PulseFactor, ShadowTint, BorderTint);
 
 		const FSlateBrush* ShadowBrush = FEditorTutorialStyle::Get().GetBrush(TEXT("Tutorials.Shadow"));
-		const FSlateBrush* BorderBrush = FEditorTutorialStyle::Get().GetBrush(TEXT("Tutorials.Border"));
+		const FSlateBrush* BorderBrush = FEditorTutorialStyle::Get().GetBrush(TEXT("Tutorials.Highlight.Border"));
 					
 		const FGeometry& WidgetGeometry = CachedGeometry;
 		const FVector2D WindowSize = OutDrawElements.GetPaintWindow()->GetSizeInScreen();
