Include header necessary for non-unity builds.

diff --git a/Engine/Plugins/Editor/GuidedTutorials/Source/IntroTutorials/Private/EditorTutorialStyle.cpp b/Engine/Plugins/Editor/GuidedTutorials/Source/IntroTutorials/Private/EditorTutorialStyle.cpp
--- a/Engine/Plugins/Editor/GuidedTutorials/Source/IntroTutorials/Private/EditorTutorialStyle.cpp
+++ b/Engine/Plugins/Editor/GuidedTutorials/Source/IntroTutorials/Private/EditorTutorialStyle.cpp
@@ -5,6 +5,7 @@
 #include "Styling/SlateStyleMacros.h"
 #include "Styling/SlateTypes.h"
 #include "Styling/AppStyle.h"
+#include "Styling/CoreStyle.h"
 
 FName FEditorTutorialStyle::StyleName("FEditorTutorialStyle");
 TUniquePtr<FEditorTutorialStyle> FEditorTutorialStyle::Instance(nullptr);
