Make a delegate on the main frame that can be used to request resources, which plugins can register to listen for.

diff --git a/Engine/Plugins/Editor/GuidedTutorials/Source/IntroTutorials/Private/IntroTutorials.cpp b/Engine/Plugins/Editor/GuidedTutorials/Source/IntroTutorials/Private/IntroTutorials.cpp
--- a/Engine/Plugins/Editor/GuidedTutorials/Source/IntroTutorials/Private/IntroTutorials.cpp
+++ b/Engine/Plugins/Editor/GuidedTutorials/Source/IntroTutorials/Private/IntroTutorials.cpp
@@ -99,7 +99,8 @@ void FIntroTutorials::StartupModule()
 		IMainFrameModule& MainFrameModule = FModuleManager::LoadModuleChecked<IMainFrameModule>(TEXT("MainFrame"));
 		MainFrameModule.OnMainFrameCreationFinished().AddRaw(this, &FIntroTutorials::MainFrameLoad);
 		MainFrameModule.OnMainFrameSDKNotInstalled().AddRaw(this, &FIntroTutorials::HandleSDKNotInstalled);
-		
+		MainFrameModule.OnMainFrameRequestResource().AddRaw(this, &FIntroTutorials::HandleRequestResource);
+
 		// Add menu option for level editor tutorial
 		RegisterSummonTutorialsMenuEntries();
 
@@ -396,6 +397,21 @@ void FIntroTutorials::HandleSDKNotInstalled(const FString& PlatformName, const F
 	}
 }
 
+void FIntroTutorials::HandleRequestResource(const FString& Category, const FString& InTutorialAsset)
+{
+	if (Category == "Tutorial")
+	{
+		if (FPackageName::IsValidLongPackageName(InTutorialAsset, true))
+		{
+			UBlueprint* Blueprint = LoadObject<UBlueprint>(nullptr, *InTutorialAsset);
+			if (Blueprint)
+			{
+				LaunchTutorialByName(InTutorialAsset);
+			}
+		}
+	}
+}
+
 EVisibility FIntroTutorials::GetHomeButtonVisibility() const
 {
 	if(CurrentObjectClass != nullptr)
diff --git a/Engine/Plugins/Editor/GuidedTutorials/Source/IntroTutorials/Private/IntroTutorials.h b/Engine/Plugins/Editor/GuidedTutorials/Source/IntroTutorials/Private/IntroTutorials.h
--- a/Engine/Plugins/Editor/GuidedTutorials/Source/IntroTutorials/Private/IntroTutorials.h
+++ b/Engine/Plugins/Editor/GuidedTutorials/Source/IntroTutorials/Private/IntroTutorials.h
@@ -100,6 +100,9 @@ private:
 	/** Events to call when SDK isn't installed */
 	void HandleSDKNotInstalled(const FString& PlatformName, const FString& InTutorialAsset);
 
+	/** Handles calls to MainFrameRequestResource that specify a tutorial asset and "Tutorial" category. */
+	void HandleRequestResource(const FString& Category, const FString& InTutorialName);
+
 	bool MaybeOpenWelcomeTutorial();
 	void ResetWelcomeTutorials() const;
 
diff --git a/Engine/Source/Developer/MessageLog/MessageLog.Build.cs b/Engine/Source/Developer/MessageLog/MessageLog.Build.cs
--- a/Engine/Source/Developer/MessageLog/MessageLog.Build.cs
+++ b/Engine/Source/Developer/MessageLog/MessageLog.Build.cs
@@ -31,6 +31,7 @@ public class MessageLog : ModuleRules
 			PrivateIncludePathModuleNames.AddRange(
 				new string[] {
 					"Documentation",
+					"MainFrame"
 				}
 			);		
 		}
diff --git a/Engine/Source/Developer/MessageLog/Private/UserInterface/SMessageLogMessageListRow.h b/Engine/Source/Developer/MessageLog/Private/UserInterface/SMessageLogMessageListRow.h
--- a/Engine/Source/Developer/MessageLog/Private/UserInterface/SMessageLogMessageListRow.h
+++ b/Engine/Source/Developer/MessageLog/Private/UserInterface/SMessageLogMessageListRow.h
@@ -15,6 +15,7 @@
 
 #if WITH_EDITOR
 	#include "IDocumentation.h"
+	#include "Interfaces/IMainFrameModule.h"
 #endif
 
 #define LOCTEXT_NAMESPACE "SMessageLogMessageListRow"
@@ -102,7 +103,8 @@ private:
 
 	void HandleTutorialHyperlinkNavigate( FString TutorialAssetName )
 	{
-		//IIntroTutorials::Get().LaunchTutorial(TutorialAssetName);
+		IMainFrameModule& MainFrameModule = FModuleManager::GetModuleChecked<IMainFrameModule>(TEXT("MainFrame"));
+		MainFrameModule.BroadcastMainFrameRequestResource("Tutorial", TutorialAssetName);
 	}
 #endif
 
diff --git a/Engine/Source/Editor/MainFrame/Private/MainFrameModule.h b/Engine/Source/Editor/MainFrame/Private/MainFrameModule.h
--- a/Engine/Source/Editor/MainFrame/Private/MainFrameModule.h
+++ b/Engine/Source/Editor/MainFrame/Private/MainFrameModule.h
@@ -122,6 +122,16 @@ public:
 		return MainFrameSDKNotInstalled.Broadcast(PlatformName, DocLink);
 	}
 
+	DECLARE_DERIVED_EVENT(FMainFrameModule, IMainFrameModule::FMainFrameRequestResource, FMainFrameRequestResource);
+	virtual FMainFrameRequestResource& OnMainFrameRequestResource() override
+	{
+		return MainFrameRequestResource;
+	}
+	void BroadcastMainFrameRequestResource(const FString& Category, const FString& ResourceName) override
+	{
+		return MainFrameRequestResource.Broadcast(Category, ResourceName);
+	}
+
 	virtual void EnableDelayedShowMainFrame() override
 	{
 		bDelayedShowMainFrame = true;
@@ -235,6 +245,9 @@ private:
 	/// Event to be called when the editor tried to use a platform, but it wasn't installed
 	FMainFrameSDKNotInstalled MainFrameSDKNotInstalled;
 
+	/// Event to be called to make an open-ended request for a resource from any registered listeners
+	FMainFrameRequestResource MainFrameRequestResource;
+
 	// Commands used by main frame in menus and key bindings.
 	TSharedPtr<class FMainFrameCommands> MainFrameActions;
 
diff --git a/Engine/Source/Editor/MainFrame/Public/Interfaces/IMainFrameModule.h b/Engine/Source/Editor/MainFrame/Public/Interfaces/IMainFrameModule.h
--- a/Engine/Source/Editor/MainFrame/Public/Interfaces/IMainFrameModule.h
+++ b/Engine/Source/Editor/MainFrame/Public/Interfaces/IMainFrameModule.h
@@ -182,6 +182,13 @@ public:
 	virtual FMainFrameSDKNotInstalled& OnMainFrameSDKNotInstalled( ) = 0;
 	virtual void BroadcastMainFrameSDKNotInstalled(const FString& PlatformName, const FString& DocLink) = 0;
 
+	/**
+	 * Delegate for making an open-ended request for a resource or link.
+	 */
+	DECLARE_EVENT_TwoParams(IMainFrameModule, FMainFrameRequestResource, const FString&, const FString&);
+	virtual FMainFrameRequestResource& OnMainFrameRequestResource() = 0;
+	virtual void BroadcastMainFrameRequestResource(const FString& Category, const FString& ResourceName) = 0;
+
 	/**
 	 * Enable external control of when main frame is shown
 	 */
