[NO NON-PII COMMIT MESSAGE]

diff --git a/Engine/Plugins/Editor/GuidedTutorials/Source/IntroTutorials/Private/IntroTutorials.cpp b/Engine/Plugins/Editor/GuidedTutorials/Source/IntroTutorials/Private/IntroTutorials.cpp
--- a/Engine/Plugins/Editor/GuidedTutorials/Source/IntroTutorials/Private/IntroTutorials.cpp
+++ b/Engine/Plugins/Editor/GuidedTutorials/Source/IntroTutorials/Private/IntroTutorials.cpp
@@ -235,7 +235,7 @@ void FIntroTutorials::RegisterSummonTutorialsMenuEntries()
 	Section.AddEntry(FToolMenuEntry::InitMenuEntry(
 		"Tutorials",
 		LOCTEXT("TutorialsMenuEntryTitle", "Tutorials"),
-		LOCTEXT("TutorialsMenuEntryToolTip", "Opens up introductory tutorials covering the basics of using the Unreal Engine 4 Editor."),
+		LOCTEXT("TutorialsMenuEntryToolTip", "Opens up introductory tutorials covering the basics of using the Unreal Engine Editor."),
 		FSlateIcon(FEditorTutorialStyle::Get().GetStyleSetName(), "Tutorials.MenuIcon"),
 		FUIAction(FExecuteAction::CreateRaw(this, &FIntroTutorials::SummonTutorialHome))
 	));
